The build scale is now taken into account when calculating a Static Meshes streaming texture factors.

#extra The user has to manually reimport the static mesh (and save) in order to process ComputeStreamingTextureFactors again and correct the streaming texture factors saved to the mesh. The level viewport will update to reflect the texture change after a couple of seconds have passed.

#ttp 332902 - LIVE: RENDERING: MipMaps blur as if using far LOD when mesh uses large Build Scale

#branch UE4

#change The LOD build scale is now applied to the wedge position as it is taken from the raw mesh.

diff --git a/Engine/Source/Developer/MeshUtilities/Private/MeshUtilities.cpp b/Engine/Source/Developer/MeshUtilities/Private/MeshUtilities.cpp
--- a/Engine/Source/Developer/MeshUtilities/Private/MeshUtilities.cpp
+++ b/Engine/Source/Developer/MeshUtilities/Private/MeshUtilities.cpp
@@ -1447,7 +1447,8 @@ static void ComputeTangents(
 static void ComputeStreamingTextureFactors(
 	float* OutStreamingTextureFactors,
 	float* OutMaxStreamingTextureFactor,
-	const FRawMesh& Mesh
+	const FRawMesh& Mesh,
+	const FVector& BuildScale
 	)
 {
 	int32 NumTexCoords = ComputeNumTexCoords(Mesh, MAX_STATIC_TEXCOORDS);
@@ -1460,9 +1461,9 @@ static void ComputeStreamingTextureFactors(
 		int32 Wedge1 = FaceIndex * 3 + 1;
 		int32 Wedge2 = FaceIndex * 3 + 2;
 
-		const FVector& Pos0 = Mesh.GetWedgePosition(Wedge0);
-		const FVector& Pos1 = Mesh.GetWedgePosition(Wedge1);
-		const FVector& Pos2 = Mesh.GetWedgePosition(Wedge2);
+		const FVector& Pos0 = Mesh.GetWedgePosition(Wedge0) * BuildScale;
+		const FVector& Pos1 = Mesh.GetWedgePosition(Wedge1) * BuildScale;
+		const FVector& Pos2 = Mesh.GetWedgePosition(Wedge2) * BuildScale;
 		float	L1	= (Pos0 - Pos1).Size(),
 				L2	= (Pos0 - Pos2).Size();
 
@@ -2035,7 +2036,8 @@ bool FMeshUtilities::BuildStaticMesh(
 			ComputeStreamingTextureFactors(
 				OutRenderData.StreamingTextureFactors,
 				&OutRenderData.MaxStreamingTextureFactor,
-				RawMesh
+				RawMesh,
+				LODBuildSettings[LODIndex].BuildScale3D
 				);
 		}
 
