Undo accidental change where we make BoneVertsInfo not use local space. Fixes the problems in PhAT
TTP#334913

diff --git a/Engine/Source/Developer/MeshUtilities/Private/SkeletalMeshTools.cpp b/Engine/Source/Developer/MeshUtilities/Private/SkeletalMeshTools.cpp
--- a/Engine/Source/Developer/MeshUtilities/Private/SkeletalMeshTools.cpp
+++ b/Engine/Source/Developer/MeshUtilities/Private/SkeletalMeshTools.cpp
@@ -309,10 +309,10 @@ namespace SkeletalMeshTools
 				FRigidSkinVertex* RigidVert = &Chunk.RigidVertices[i];
 				int32 BoneIndex = Chunk.BoneMap[RigidVert->Bone];
 
-				FVector LocalPos = RigidVert->Position;
+				FVector LocalPos = SkeletalMesh->RefBasesInvMatrix[BoneIndex].TransformPosition(RigidVert->Position);
 				Infos[BoneIndex].Positions.Add(LocalPos);
 
-				FVector LocalNormal = RigidVert->TangentZ;
+				FVector LocalNormal = SkeletalMesh->RefBasesInvMatrix[BoneIndex].TransformVector(RigidVert->TangentZ);
 				Infos[BoneIndex].Normals.Add(LocalNormal);
 			}
 
@@ -324,10 +324,10 @@ namespace SkeletalMeshTools
 				{
 					int32 BoneIndex = Chunk.BoneMap[GetDominantBoneIndex(SoftVert)];
 
-					FVector LocalPos = SoftVert->Position;
+					FVector LocalPos = SkeletalMesh->RefBasesInvMatrix[BoneIndex].TransformPosition(SoftVert->Position);
 					Infos[BoneIndex].Positions.Add(LocalPos);
 
-					FVector LocalNormal = SoftVert->TangentZ;
+					FVector LocalNormal = SkeletalMesh->RefBasesInvMatrix[BoneIndex].TransformVector(SoftVert->TangentZ);
 					Infos[BoneIndex].Normals.Add(LocalNormal);
 				}
 				else
@@ -338,10 +338,10 @@ namespace SkeletalMeshTools
 						{
 							int32 BoneIndex = Chunk.BoneMap[SoftVert->InfluenceBones[j]];
 
-							FVector LocalPos = SoftVert->Position;
+							FVector LocalPos = SkeletalMesh->RefBasesInvMatrix[BoneIndex].TransformPosition(SoftVert->Position);
 							Infos[BoneIndex].Positions.Add(LocalPos);
 
-							FVector LocalNormal = SoftVert->TangentZ;
+							FVector LocalNormal = SkeletalMesh->RefBasesInvMatrix[BoneIndex].TransformVector(SoftVert->TangentZ);
 							Infos[BoneIndex].Normals.Add(LocalNormal);
 						}
 					}
