Fixed: MeshMerging user settings was not used during merging

diff --git a/Engine/Source/Developer/MeshUtilities/Private/MeshUtilities.cpp b/Engine/Source/Developer/MeshUtilities/Private/MeshUtilities.cpp
--- a/Engine/Source/Developer/MeshUtilities/Private/MeshUtilities.cpp
+++ b/Engine/Source/Developer/MeshUtilities/Private/MeshUtilities.cpp
@@ -3496,19 +3496,30 @@ void FMeshUtilities::MergeActors(
 	// Pack lightmaps
 	static const uint32 MaxLightmapRes = 2048;
 	float MergedLightmapScale = 1.f;
+	FLightmapPacker LightmapPacker;	
+	
 	TArray<uint32> LightmapResList;
 	for (const FRawMeshExt& SourceMesh : SourceMeshes)
 	{
 		LightmapResList.Add(SourceMesh.LightMapRes);
 	}
 							
-	FLightmapPacker LightmapPacker;	
-	LightmapPacker.Pack(LightmapResList);
-	MergedMesh.LightMapRes = LightmapPacker.GetAtlasResolution();
-	if (MergedMesh.LightMapRes > MaxLightmapRes)
+	if (InSettings.bGnerateAtlasedLightmapUV)
 	{
-		MergedLightmapScale = MaxLightmapRes/(float)MergedMesh.LightMapRes;
-		MergedMesh.LightMapRes = MaxLightmapRes;
+		LightmapPacker.Pack(LightmapResList);
+		MergedMesh.LightMapRes = LightmapPacker.GetAtlasResolution();
+
+		if (MergedMesh.LightMapRes > MaxLightmapRes)
+		{
+			MergedLightmapScale = MaxLightmapRes/(float)MergedMesh.LightMapRes;
+			MergedMesh.LightMapRes = MaxLightmapRes;
+		}
+	}
+	else
+	{
+		// Find maximum lightmap resolution
+		LightmapResList.Sort();
+		MergedMesh.LightMapRes = LightmapResList.Last();
 	}
 
 	// Use first mesh for naming and pivot
@@ -3576,9 +3587,13 @@ void FMeshUtilities::MergeActors(
 		// Transform lightmap UVs
 		if (MergedMesh.LightMapRes)
 		{
+			FVector2D	UVOffset = FVector2D::ZeroVector;
 			FIntRect	PackedLightmapRect = LightmapPacker.GetPackedLightmapRect(SourceMeshIdx);
-			FVector2D	UVOffset = FVector2D(PackedLightmapRect.Min) * MergedLightmapScale / MergedMesh.LightMapRes;
-
+			if (PackedLightmapRect.Area() != 0)
+			{
+				UVOffset = FVector2D(PackedLightmapRect.Min) * MergedLightmapScale / MergedMesh.LightMapRes;
+			}
+			
 			for (FVector2D LightMapUV : SourceRawMesh.WedgeTexCoords[SourceMeshes[SourceMeshIdx].LightMapCoordinateIndex])
 			{
 				float UVScale = SourceMeshes[SourceMeshIdx].LightMapRes*MergedLightmapScale/MergedMesh.LightMapRes;
diff --git a/Engine/Source/Editor/UnrealEd/Private/Dialogs/SMeshMergingDialog.cpp b/Engine/Source/Editor/UnrealEd/Private/Dialogs/SMeshMergingDialog.cpp
--- a/Engine/Source/Editor/UnrealEd/Private/Dialogs/SMeshMergingDialog.cpp
+++ b/Engine/Source/Editor/UnrealEd/Private/Dialogs/SMeshMergingDialog.cpp
@@ -355,10 +355,9 @@ void SMeshMergingDialog::RunMerging()
 		return;
 	}
 	
-	FMeshMergingSettings MergeSettings;
 	FVector MergedActorLocation;
 	TArray<UObject*> AssetsToSync;
-	MeshUtilities.MergeActors(Actors, MergeSettings, MergedMeshPackageName, AssetsToSync, MergedActorLocation);
+	MeshUtilities.MergeActors(Actors, MergingSettings, MergedMeshPackageName, AssetsToSync, MergedActorLocation);
 
 	if (AssetsToSync.Num())
 	{
