MeshMerging now merges all UV channels. Added UI for specifiing target lightmap channel index

diff --git a/Engine/Source/Developer/MeshUtilities/Private/MeshUtilities.cpp b/Engine/Source/Developer/MeshUtilities/Private/MeshUtilities.cpp
--- a/Engine/Source/Developer/MeshUtilities/Private/MeshUtilities.cpp
+++ b/Engine/Source/Developer/MeshUtilities/Private/MeshUtilities.cpp
@@ -3417,7 +3417,7 @@ void FMeshUtilities::MergeActors(
 	struct FRawMeshExt
 	{
 		FRawMeshExt() 
-			: LightMapCoordinateIndex(1)
+			: LightMapCoordinateIndex(0)
 			, LightMapRes(32)
 		{}
 		
@@ -3432,6 +3432,7 @@ void FMeshUtilities::MergeActors(
 	TMap<int32, TArray<int32>>						MaterialMap;
 	TArray<FRawMeshExt>								SourceMeshes;
 	bool											bWithVertexColors = false;
+	bool											bOcuppiedUVChannels[MAX_MESH_TEXTURE_COORDS] = {};
 	
 	// Convert collected static mesh components into raw meshes
 	SourceMeshes.Empty(ComponentsToMerge.Num());
@@ -3457,7 +3458,14 @@ void FMeshUtilities::MergeActors(
 			// Source mesh asset package name
 			SourceMeshes[MeshId].AssetPackageName = MeshComponent->StaticMesh->GetOutermost()->GetName();
 
+			// Whether at least one of the meshes has vertex color
 			bWithVertexColors|= (SourceMeshes[MeshId].Mesh.WedgeColors.Num() != 0);
+
+			// Which UV channels has data at least in one mesh
+			for (int32 ChannelIdx = 0; ChannelIdx < MAX_MESH_TEXTURE_COORDS; ++ChannelIdx)
+			{
+				bOcuppiedUVChannels[ChannelIdx]|= (SourceMeshes[MeshId].Mesh.WedgeTexCoords[ChannelIdx].Num() != 0);
+			}
 		}
 		else
 		{
@@ -3490,37 +3498,35 @@ void FMeshUtilities::MergeActors(
 
 	FRawMeshExt MergedMesh;
 
-	// Set target channel for lightmap UV
-	MergedMesh.LightMapCoordinateIndex = InSettings.TargetLightmapUVChannel;
-
 	// Pack lightmaps
-	static const uint32 MaxLightmapRes = 2048;
 	float MergedLightmapScale = 1.f;
 	FLightmapPacker LightmapPacker;	
-	
-	TArray<uint32> LightmapResList;
-	for (const FRawMeshExt& SourceMesh : SourceMeshes)
-	{
-		LightmapResList.Add(SourceMesh.LightMapRes);
-	}
-							
+					
 	if (InSettings.bGnerateAtlasedLightmapUV)
 	{
+		static const uint32 MaxLightmapRes = 2048;
+				
+		// Set target channel for lightmap UV
+		MergedMesh.LightMapCoordinateIndex = InSettings.TargetLightmapUVChannel;
+		
+		// Collect lightmap sizes from all meshes
+		TArray<uint32> LightmapResList;
+		for (const FRawMeshExt& SourceMesh : SourceMeshes)
+		{
+			LightmapResList.Add(SourceMesh.LightMapRes);
+		}
+		
+		// Pack them into one atlas
 		LightmapPacker.Pack(LightmapResList);
 		MergedMesh.LightMapRes = LightmapPacker.GetAtlasResolution();
 
+		// Whether we need to scale down UV coordiantes
 		if (MergedMesh.LightMapRes > MaxLightmapRes)
 		{
 			MergedLightmapScale = MaxLightmapRes/(float)MergedMesh.LightMapRes;
 			MergedMesh.LightMapRes = MaxLightmapRes;
 		}
 	}
-	else
-	{
-		// Find maximum lightmap resolution
-		LightmapResList.Sort();
-		MergedMesh.LightMapRes = LightmapResList.Last();
-	}
 
 	// Use first mesh for naming and pivot
 	MergedMesh.AssetPackageName = SourceMeshes[0].AssetPackageName;
@@ -3569,23 +3575,9 @@ void FMeshUtilities::MergeActors(
 				FMemory::Memset(&TargetRawMesh.WedgeColors[ColorsOffset], 0xFF, ColorsNum*TargetRawMesh.WedgeColors.GetTypeSize());
 			}
 		}
-
-		if (InSettings.bImportAllUVChannels)
-		{
-	/*		const int32 NumChannels = ARRAY_COUNT(TargetRawMesh.WedgeTexCoords);
-			for (int32 ChannelIdx = 0; ChannelIdx < NumChannels; ++ChannelIdx)
-			{
-				int32 TargetChannelIdx+= (ChannelIdx >= MergedMesh.LightMapCoordinateIndex ?)
-			}*/
-			TargetRawMesh.WedgeTexCoords[0].Append(SourceRawMesh.WedgeTexCoords[0]);
-		}
-		else // Only first UV channel will be used 
-		{
-			TargetRawMesh.WedgeTexCoords[0].Append(SourceRawMesh.WedgeTexCoords[0]);
-		}
-				
-		// Transform lightmap UVs
-		if (MergedMesh.LightMapRes)
+		
+		// Write atlased UVs into user specified TargetLightmapChannel 
+		if (InSettings.bGnerateAtlasedLightmapUV && MergedMesh.LightMapRes)
 		{
 			FVector2D	UVOffset = FVector2D::ZeroVector;
 			FIntRect	PackedLightmapRect = LightmapPacker.GetPackedLightmapRect(SourceMeshIdx);
@@ -3594,12 +3586,47 @@ void FMeshUtilities::MergeActors(
 				UVOffset = FVector2D(PackedLightmapRect.Min) * MergedLightmapScale / MergedMesh.LightMapRes;
 			}
 			
-			for (FVector2D LightMapUV : SourceRawMesh.WedgeTexCoords[SourceMeshes[SourceMeshIdx].LightMapCoordinateIndex])
+			const TArray<FVector2D>& SourceWedgeTexCoords = SourceRawMesh.WedgeTexCoords[SourceMeshes[SourceMeshIdx].LightMapCoordinateIndex];
+			
+			for (FVector2D LightMapUV : SourceWedgeTexCoords)
 			{
-				float UVScale = SourceMeshes[SourceMeshIdx].LightMapRes*MergedLightmapScale/MergedMesh.LightMapRes;
+				const float SourceMeshLightmapRes = SourceMeshes[SourceMeshIdx].LightMapRes;
+				const float UVScale = (SourceMeshLightmapRes * MergedLightmapScale) / MergedMesh.LightMapRes;
 				TargetRawMesh.WedgeTexCoords[MergedMesh.LightMapCoordinateIndex].Add(LightMapUV * UVScale + UVOffset);
 			}
 		}
+		
+		// Merge all other UV channels 
+		for (int32 ChannelIdx = 0; ChannelIdx < MAX_MESH_TEXTURE_COORDS; ++ChannelIdx)
+		{
+			// Skip Lightmap channel if any
+			if (InSettings.bGnerateAtlasedLightmapUV && ChannelIdx == MergedMesh.LightMapCoordinateIndex)
+			{
+				continue;
+			}
+			
+			// Whether this channel has data
+			if (bOcuppiedUVChannels[ChannelIdx])
+			{
+				const TArray<FVector2D>& SourceChannel = SourceRawMesh.WedgeTexCoords[ChannelIdx];
+				TArray<FVector2D>& TargetChannel = TargetRawMesh.WedgeTexCoords[ChannelIdx];
+
+				// Whether source mesh has data in this channel
+				if (SourceChannel.Num())
+				{
+					TargetChannel.Append(SourceChannel);
+				}
+				else
+				{
+					// Fill with zero coordinates if source mesh has no data for this channel
+					const int32 TexCoordNum = SourceRawMesh.WedgeIndices.Num();
+					for (int32 CoordIdx = 0; CoordIdx < TexCoordNum; ++CoordIdx)
+					{
+						TargetChannel.Add(FVector2D::ZeroVector);
+					}
+				}
+			}
+		}
 	}
 	
 	//
@@ -3633,8 +3660,11 @@ void FMeshUtilities::MergeActors(
 		StaticMesh->LightingGuid = FGuid::NewGuid();
 
 		// Set it to use textured lightmaps. Note that Build Lighting will do the error-checking (texcoordindex exists for all LODs, etc).
-		StaticMesh->LightMapResolution = MergedMesh.LightMapRes;
-		StaticMesh->LightMapCoordinateIndex = MergedMesh.LightMapCoordinateIndex;
+		if (InSettings.bGnerateAtlasedLightmapUV)
+		{
+			StaticMesh->LightMapResolution = MergedMesh.LightMapRes;
+			StaticMesh->LightMapCoordinateIndex = MergedMesh.LightMapCoordinateIndex;	
+		}
 
 		FStaticMeshSourceModel* SrcModel = new (StaticMesh->SourceModels) FStaticMeshSourceModel();
 		/*Don't allow the engine to recalculate normals*/
diff --git a/Engine/Source/Editor/UnrealEd/Private/Dialogs/SMeshMergingDialog.cpp b/Engine/Source/Editor/UnrealEd/Private/Dialogs/SMeshMergingDialog.cpp
--- a/Engine/Source/Editor/UnrealEd/Private/Dialogs/SMeshMergingDialog.cpp
+++ b/Engine/Source/Editor/UnrealEd/Private/Dialogs/SMeshMergingDialog.cpp
@@ -3,11 +3,13 @@
 #include "UnrealEd.h"
 #include "Dialogs/DlgPickAssetPath.h"
 #include "Dialogs/SMeshMergingDialog.h"
+#include "RawMesh.h"
 #include "MeshUtilities.h"
 #include "ContentBrowserModule.h"
 #include "AssetRegistryModule.h"
 #include "Editor/LevelEditor/Public/LevelEditor.h"
 
+
 #define LOCTEXT_NAMESPACE "SMeshMergingDialog"
 
 //////////////////////////////////////////////////////////////////////////
@@ -33,8 +35,6 @@ void SMeshMergingDialog::Construct(const FArguments& InArgs)
 
 	GenerateNewPackageName();
 
-	FMeshMergingSettings();
-
 	// Create widget layout
 	this->ChildSlot
 	[
@@ -69,7 +69,29 @@ void SMeshMergingDialog::Construct(const FArguments& InArgs)
 							SNew(STextBlock).Text(LOCTEXT("AtlasLightmapUVLabel", "Generate Atlased Lightmap UVs"))
 						]
 					]
-
+					
+					// Target lightmap channel
+					+SVerticalBox::Slot()
+					.AutoHeight()
+					.Padding(FEditorStyle::GetMargin("StandardDialog.ContentPadding"))
+					[
+						SNew(SBorder)
+						.Padding(FMargin(8,0,0,0))
+						.BorderImage( FEditorStyle::GetBrush("ToolPanel.GroupBorder"))
+						[
+							SNew(SNumericEntryBox<int32>)
+							.IsEnabled( this, &SMeshMergingDialog::IsLightmapChannelEnabled )
+							.MinValue(0)
+							.MaxValue(MAX_MESH_TEXTURE_COORDS-1)
+							.Value(this, &SMeshMergingDialog::GetTargetLightmapChannelValue)
+							.OnValueCommitted(this, &SMeshMergingDialog::OnTargetLightmapChannelValueCommited)
+							.Label()
+							[
+								SNumericEntryBox<int32>::BuildLabel(LOCTEXT("TargetLightMapChannelLabel", "Target LightMap Channel: "), FLinearColor::Black, FLinearColor(.33f,.33f,.33f))
+							]
+						]
+					]
+					
 					// Vertex colors
 					+SVerticalBox::Slot()
 					.AutoHeight()
@@ -222,6 +244,21 @@ void SMeshMergingDialog::SetAtlasLightmapUV(ESlateCheckBoxState::Type NewValue)
 	MergingSettings.bGnerateAtlasedLightmapUV = (ESlateCheckBoxState::Checked == NewValue);
 }
 
+bool SMeshMergingDialog::IsLightmapChannelEnabled() const
+{
+	return MergingSettings.bGnerateAtlasedLightmapUV;
+}
+
+TOptional<int32> SMeshMergingDialog::GetTargetLightmapChannelValue() const
+{
+	return MergingSettings.TargetLightmapUVChannel;
+}
+
+void SMeshMergingDialog::OnTargetLightmapChannelValueCommited(int32 NewValue, ETextCommit::Type CommitInfo)
+{
+	MergingSettings.TargetLightmapUVChannel = FMath::Clamp(NewValue, 0, MAX_MESH_TEXTURE_COORDS-1);
+}
+
 ESlateCheckBoxState::Type SMeshMergingDialog::GetImportVertexColors() const
 {
 	return (MergingSettings.bImportVertexColors ? ESlateCheckBoxState::Checked : ESlateCheckBoxState::Unchecked);
diff --git a/Engine/Source/Editor/UnrealEd/Public/Dialogs/SMeshMergingDialog.h b/Engine/Source/Editor/UnrealEd/Public/Dialogs/SMeshMergingDialog.h
--- a/Engine/Source/Editor/UnrealEd/Public/Dialogs/SMeshMergingDialog.h
+++ b/Engine/Source/Editor/UnrealEd/Public/Dialogs/SMeshMergingDialog.h
@@ -37,6 +37,11 @@ private:
 	ESlateCheckBoxState::Type GetAtlasLightmapUV() const;
 	void SetAtlasLightmapUV(ESlateCheckBoxState::Type NewValue);
 
+	/** Target lightmap channel */
+	bool IsLightmapChannelEnabled() const;
+	TOptional<int32> GetTargetLightmapChannelValue() const;
+	void OnTargetLightmapChannelValueCommited(int32 NewValue, ETextCommit::Type CommitInfo);
+	
 	/**  */
 	ESlateCheckBoxState::Type GetImportVertexColors() const;
 	void SetImportVertexColors(ESlateCheckBoxState::Type NewValue);
