Fixed console variable needed by [REDACTED_PERSON], which had been broken by console manager refactor

diff --git a/Engine/Source/Developer/MeshUtilities/Private/MeshUtilities.cpp b/Engine/Source/Developer/MeshUtilities/Private/MeshUtilities.cpp
--- a/Engine/Source/Developer/MeshUtilities/Private/MeshUtilities.cpp
+++ b/Engine/Source/Developer/MeshUtilities/Private/MeshUtilities.cpp
@@ -41,13 +41,6 @@ static TAutoConsoleVariable<int32> CVarTriangleOrderOptimization(
 	TEXT("1: Use Forsyth algorithm (fastest)(default)"),
 	ECVF_ReadOnly);
 
-static TAutoConsoleVariable<int32> CVarAllowMeshDistanceFieldRepresentations (
-	TEXT("r.AllowMeshDistanceFieldRepresentations"),
-	0,
-	TEXT("Whether to build distance fields of static meshes, needed for distance field AO, which is used to implement Movable SkyLight shadows.\n")
-	TEXT("Enabling will increase mesh build times and memory usage.  Changing this value will cause a rebuild of all static meshes."),
-	ECVF_ReadOnly);
-
 //@todo - implement required vector intrinsics for other implementations
 #if PLATFORM_ENABLE_VECTORINTRINSICS
 
@@ -266,8 +259,10 @@ static void GenerateSignedDistanceFieldVolumeData(
 	const FBoxSphereBounds& Bounds,
 	float DistanceFieldResolutionScale)
 {
+	static const auto CVar = IConsoleManager::Get().FindTConsoleVariableDataInt(TEXT("r.AllowMeshDistanceFieldRepresentations"));
+
 	if (DistanceFieldResolutionScale > 0 
-		&& CVarAllowMeshDistanceFieldRepresentations.GetValueOnGameThread() != 0)
+		&& CVar->GetValueOnGameThread() != 0)
 	{
 		const FPositionVertexBuffer& PositionVertexBuffer = LODModel.PositionVertexBuffer;
 		FIndexArrayView Indices = LODModel.IndexBuffer.GetArrayView();
diff --git a/Engine/Source/Runtime/Core/Private/HAL/ConsoleManager.cpp b/Engine/Source/Runtime/Core/Private/HAL/ConsoleManager.cpp
--- a/Engine/Source/Runtime/Core/Private/HAL/ConsoleManager.cpp
+++ b/Engine/Source/Runtime/Core/Private/HAL/ConsoleManager.cpp
@@ -1575,6 +1575,13 @@ static TAutoConsoleVariable<int32> CVarAllowStaticLighting(
 	TEXT("Games that only use dynamic lighting should set this to 0 to save some static lighting overhead."),
 	ECVF_ReadOnly | ECVF_RenderThreadSafe);
 
+static TAutoConsoleVariable<int32> CVarDistField(
+	TEXT("r.AllowMeshDistanceFieldRepresentations"),
+	0,	
+	TEXT("Whether to build distance fields of static meshes, needed for distance field AO, which is used to implement Movable SkyLight shadows.\n")
+	TEXT("Enabling will increase mesh build times and memory usage.  Changing this value will cause a rebuild of all static meshes."),
+	ECVF_ReadOnly);
+
 static TAutoConsoleVariable<int32> CVarMinLogVerbosity(
 	TEXT("con.MinLogVerbosity"),
 	0,
