Fixed normals calculations during mesh merging and proxy mesh generation
TTP# 340708

diff --git a/Engine/Source/Developer/MeshUtilities/Private/MeshUtilities.cpp b/Engine/Source/Developer/MeshUtilities/Private/MeshUtilities.cpp
--- a/Engine/Source/Developer/MeshUtilities/Private/MeshUtilities.cpp
+++ b/Engine/Source/Developer/MeshUtilities/Private/MeshUtilities.cpp
@@ -3241,8 +3241,6 @@ bool FMeshUtilities::ConstructRawMesh(
 
 	//Transform the raw mesh to world space
 	FTransform CtoM = InMeshComponent->ComponentToWorld;
-	FMatrix InvTransCToM = CtoM.ToMatrixWithScale().Inverse().GetTransposed();
-
 	for (FVector& Vertex : OutRawMesh.VertexPositions)
 	{
 		Vertex = CtoM.TransformFVector4(Vertex);
@@ -3261,13 +3259,6 @@ bool FMeshUtilities::ConstructRawMesh(
 	FindOverlappingCorners(OverlappingCorners, OutRawMesh, 0.1f);
 	ComputeTangents(OutRawMesh, OverlappingCorners, ETangentOptions::BlendOverlappingNormals);
 
-	for (int32 WedgeIndex = 0; WedgeIndex < NumWedges; ++WedgeIndex)
-	{
-		OutRawMesh.WedgeTangentX[WedgeIndex] = InvTransCToM.TransformVector(OutRawMesh.WedgeTangentX[WedgeIndex]).SafeNormal();
-		OutRawMesh.WedgeTangentY[WedgeIndex] = InvTransCToM.TransformVector(OutRawMesh.WedgeTangentY[WedgeIndex]).SafeNormal();
-		OutRawMesh.WedgeTangentZ[WedgeIndex] = InvTransCToM.TransformVector(OutRawMesh.WedgeTangentZ[WedgeIndex]).SafeNormal();
-	}
-
 	//Need to store the unique material indices in order to re-map the material indices in each rawmesh
 	//Only using the base mesh
 	for (const FStaticMeshSection& Section : SrcMesh->RenderData->LODResources[0].Sections) 
