Added painted vertex color support to MeshMerging tool

diff --git a/Engine/Source/Developer/MeshUtilities/Private/MeshUtilities.cpp b/Engine/Source/Developer/MeshUtilities/Private/MeshUtilities.cpp
--- a/Engine/Source/Developer/MeshUtilities/Private/MeshUtilities.cpp
+++ b/Engine/Source/Developer/MeshUtilities/Private/MeshUtilities.cpp
@@ -3383,6 +3383,55 @@ private:
 	TArray<const FAtlasedTextureSlot*>		PackedLigthmapSlots;
 };
 
+bool PropagatePaintedColorsToRawMesh(UStaticMeshComponent* StaticMeshComponent, int32 LODIndex, FRawMesh& RawMesh)
+{
+	UStaticMesh* StaticMesh = StaticMeshComponent->StaticMesh;
+	
+	if (StaticMesh->SourceModels.IsValidIndex(LODIndex) &&
+		StaticMeshComponent->LODData.IsValidIndex(LODIndex) &&
+		StaticMeshComponent->LODData[LODIndex].OverrideVertexColors != nullptr)	
+	{
+		FColorVertexBuffer& ColorVertexBuffer = *StaticMeshComponent->LODData[LODIndex].OverrideVertexColors;
+		FStaticMeshSourceModel& SrcModel = StaticMesh->SourceModels[LODIndex];
+		FStaticMeshRenderData& RenderData = *StaticMesh->RenderData;
+		FStaticMeshLODResources& RenderModel = RenderData.LODResources[LODIndex];
+
+		if (RenderData.WedgeMap.Num() > 0 && 
+			ColorVertexBuffer.GetNumVertices() == RenderModel.GetNumVertices())
+		{
+			int32 NumWedges = RawMesh.WedgeIndices.Num();
+			if (RenderData.WedgeMap.Num() == NumWedges)
+			{
+				int32 NumExistingColors = RawMesh.WedgeColors.Num();
+				if (NumExistingColors < NumWedges)
+				{
+					RawMesh.WedgeColors.AddUninitialized(NumWedges - NumExistingColors);
+				}
+			
+				for (int32 i = 0; i < NumWedges; ++i)
+				{
+					FColor WedgeColor = FColor::White;
+					int32 Index = RenderData.WedgeMap[i];
+					if (Index != INDEX_NONE)
+					{
+						WedgeColor = ColorVertexBuffer.VertexColor(Index);
+					}
+				
+					RawMesh.WedgeColors[i] = WedgeColor;
+				}
+			
+				return true;
+			}
+			else
+			{
+				UE_LOG(LogMeshUtilities, Warning, TEXT("Wedge map size %d is wrong. Expected %d."), RenderData.WedgeMap.Num(), RawMesh.WedgeIndices.Num());
+			}
+		}
+	}
+	
+	return false;
+}
+
 void FMeshUtilities::MergeActors(
 	const TArray<AActor*>& SourceActors, 
 	const FMeshMergingSettings& InSettings, 
@@ -3444,9 +3493,15 @@ void FMeshUtilities::MergeActors(
 			// Source mesh asset package name
 			SourceMeshes[MeshId].AssetPackageName = MeshComponent->StaticMesh->GetOutermost()->GetName();
 
-			// Whether at least one of the meshes has vertex color
-			bWithVertexColors|= (SourceMeshes[MeshId].Mesh.WedgeColors.Num() != 0);
-
+			// Should we use vertex colors?
+			if (InSettings.bImportVertexColors)
+			{
+				// Propagate painted vertex colors into our raw mesh
+				PropagatePaintedColorsToRawMesh(MeshComponent, 0, SourceMeshes[MeshId].Mesh);
+				// Whether at least one of the meshes has vertex colors
+				bWithVertexColors|= (SourceMeshes[MeshId].Mesh.WedgeColors.Num() != 0);
+			}
+			
 			// Which UV channels has data at least in one mesh
 			for (int32 ChannelIdx = 0; ChannelIdx < MAX_MESH_TEXTURE_COORDS; ++ChannelIdx)
 			{
@@ -3463,12 +3518,6 @@ void FMeshUtilities::MergeActors(
 	{
 		return;	
 	}
-
-	// Should we use vertex colors?
-	if (!InSettings.bImportVertexColors)
-	{
-		bWithVertexColors = false;
-	}
 	
 	//For each raw mesh, re-map the material indices according to the MaterialMap
 	for (int32 MeshIndex = 0; MeshIndex < SourceMeshes.Num(); ++MeshIndex)
