Better logging when merging painted vertex colors

diff --git a/Engine/Source/Developer/MeshUtilities/Private/MeshUtilities.cpp b/Engine/Source/Developer/MeshUtilities/Private/MeshUtilities.cpp
--- a/Engine/Source/Developer/MeshUtilities/Private/MeshUtilities.cpp
+++ b/Engine/Source/Developer/MeshUtilities/Private/MeshUtilities.cpp
@@ -3424,7 +3424,7 @@ bool PropagatePaintedColorsToRawMesh(UStaticMeshComponent* StaticMeshComponent,
 			}
 			else
 			{
-				UE_LOG(LogMeshUtilities, Warning, TEXT("Wedge map size %d is wrong. Expected %d."), RenderData.WedgeMap.Num(), RawMesh.WedgeIndices.Num());
+				UE_LOG(LogMeshUtilities, Warning, TEXT("{%s} Wedge map size %d is wrong. Expected %d."), *StaticMesh->GetName(), RenderData.WedgeMap.Num(), RawMesh.WedgeIndices.Num());
 			}
 		}
 	}
