TTP 340843	CRASH: FBX IMPORT: Asset crashes engine to desktop on import

diff --git a/Engine/Source/Developer/MeshUtilities/Private/MeshUtilities.cpp b/Engine/Source/Developer/MeshUtilities/Private/MeshUtilities.cpp
--- a/Engine/Source/Developer/MeshUtilities/Private/MeshUtilities.cpp
+++ b/Engine/Source/Developer/MeshUtilities/Private/MeshUtilities.cpp
@@ -482,7 +482,7 @@ private:
 		const FStaticMeshLODGroup& LODGroup
 		) override;
 
-	virtual bool BuildSkeletalMesh( FStaticLODModel& LODModel, const FReferenceSkeleton& RefSkeleton, const TArray<FVertInfluence>& Influences, const TArray<FMeshWedge>& Wedges, const TArray<FMeshFace>& Faces, const TArray<FVector>& Points, const TArray<int32>& PointToOriginalMap, bool bKeepOverlappingVertices = false, bool bComputeNormals = true, bool bComputeTangents = true );
+	virtual bool BuildSkeletalMesh( FStaticLODModel& LODModel, const FReferenceSkeleton& RefSkeleton, const TArray<FVertInfluence>& Influences, const TArray<FMeshWedge>& Wedges, const TArray<FMeshFace>& Faces, const TArray<FVector>& Points, const TArray<int32>& PointToOriginalMap, bool bKeepOverlappingVertices = false, bool bComputeNormals = true, bool bComputeTangents = true, TArray<FText> * OutWarningMessages = NULL);
 
 	virtual bool GenerateUVs(
 		FRawMesh& RawMesh,
@@ -2551,7 +2551,7 @@ bool FMeshUtilities::BuildStaticMesh(
 }
 
 
-bool FMeshUtilities::BuildSkeletalMesh( FStaticLODModel& LODModel, const FReferenceSkeleton& RefSkeleton, const TArray<FVertInfluence>& Influences, const TArray<FMeshWedge>& Wedges, const TArray<FMeshFace>& Faces, const TArray<FVector>& Points, const TArray<int32>& PointToOriginalMap, bool bKeepOverlappingVertices, bool bComputeNormals, bool bComputeTangents )
+bool FMeshUtilities::BuildSkeletalMesh( FStaticLODModel& LODModel, const FReferenceSkeleton& RefSkeleton, const TArray<FVertInfluence>& Influences, const TArray<FMeshWedge>& Wedges, const TArray<FMeshFace>& Faces, const TArray<FVector>& Points, const TArray<int32>& PointToOriginalMap, bool bKeepOverlappingVertices, bool bComputeNormals, bool bComputeTangents, TArray<FText> * OutWarningMessages )
 {
 #if WITH_EDITORONLY_DATA
 	bool bTooManyVerts = false;
@@ -2624,9 +2624,21 @@ bool FMeshUtilities::BuildSkeletalMesh( FStaticLODModel& LODModel, const FRefere
 	{
 		uint32* InfluenceIndex = VertexIndexToInfluenceIndexMap.Find( Wedges[WedgeIndex].iVertex );
 
-		check( InfluenceIndex );
+		if ( InfluenceIndex )
+		{
+			WedgeInfluenceIndices.Add( *InfluenceIndex );
+		}
+		else
+		{
+			// we have missing influence vert,  we weight to root
+			WedgeInfluenceIndices.Add(0);
 
-		WedgeInfluenceIndices.Add( *InfluenceIndex );
+			// add warning message
+			if (OutWarningMessages)
+			{
+				OutWarningMessages->Add(FText::Format(FText::FromString("Missing influence on vert {0}. Weighting it to root."), FText::FromString(FString::FromInt(Wedges[WedgeIndex].iVertex))));
+			}
+		}
 	}
 
  	check(Wedges.Num() == WedgeInfluenceIndices.Num());
@@ -2970,19 +2982,43 @@ bool FMeshUtilities::BuildSkeletalMesh( FStaticLODModel& LODModel, const FRefere
 		}
 		if( bHasBadSections )
 		{
-			FMessageDialog::Open( EAppMsgType::Ok, NSLOCTEXT("UnrealEd", "Error_SkeletalMeshHasBadSections", "Input mesh has a section with no triangles.  This mesh may not render properly.") );
+			FText BadSectionMessage(NSLOCTEXT("UnrealEd", "Error_SkeletalMeshHasBadSections", "Input mesh has a section with no triangles.  This mesh may not render properly."));
+			if(OutWarningMessages)
+			{
+				OutWarningMessages->Add(BadSectionMessage);
+			}
+			else
+			{
+				FMessageDialog::Open( EAppMsgType::Ok, BadSectionMessage );
+			}
 		}
 
 		if (bTooManyVerts)
 		{
-			UE_LOG(LogSkeletalMesh, Log, TEXT("Input mesh has too many vertices.  The generated mesh will be corrupt!"));
-			FMessageDialog::Open( EAppMsgType::Ok, NSLOCTEXT("UnrealEd", "Error_SkeletalMeshTooManyVertices", "Input mesh has too many vertices.  The generated mesh will be corrupt!  Consider adding extra materials to split up the source mesh into smaller chunks."));
+			FText TooManyVertsMessage(NSLOCTEXT("UnrealEd", "Error_SkeletalMeshTooManyVertices", "Input mesh has too many vertices.  The generated mesh will be corrupt!  Consider adding extra materials to split up the source mesh into smaller chunks."));
+
+			if(OutWarningMessages)
+			{
+				OutWarningMessages->Add(TooManyVertsMessage);
+			}
+			else
+			{
+				FMessageDialog::Open(EAppMsgType::Ok, TooManyVertsMessage);
+			}
 		}
 	}
 
 	return true;
 #else
-	UE_LOG(LogSkeletalMesh, Fatal,TEXT("Cannot call FSkeletalMeshTools::CreateSkinningStreams on a console!"));
+
+	if(OutWarningMessages)
+	{
+		OutWarningMessages->Add(FText::FromString(TEXT("Cannot call FSkeletalMeshTools::CreateSkinningStreams on a console!")));
+	}
+	else
+	{
+		UE_LOG(LogSkeletalMesh, Fatal,TEXT("Cannot call FSkeletalMeshTools::CreateSkinningStreams on a console!"));
+	}
 	return false;
 #endif
 }
diff --git a/Engine/Source/Developer/MeshUtilities/Public/MeshUtilities.h b/Engine/Source/Developer/MeshUtilities/Public/MeshUtilities.h
--- a/Engine/Source/Developer/MeshUtilities/Public/MeshUtilities.h
+++ b/Engine/Source/Developer/MeshUtilities/Public/MeshUtilities.h
@@ -149,7 +149,8 @@ public:
 		const TArray<int32>& PointToOriginalMap,
 		bool bKeepOverlappingVertices = false,
 		bool bComputeNormals = true,
-		bool bComputeTangents = true
+		bool bComputeTangents = true,
+		TArray<FText> * OutWarningMessages = NULL
 		) = 0;
 
 
diff --git a/Engine/Source/Editor/UnrealEd/Private/Fbx/FbxSkeletalMeshImport.cpp b/Engine/Source/Editor/UnrealEd/Private/Fbx/FbxSkeletalMeshImport.cpp
--- a/Engine/Source/Editor/UnrealEd/Private/Fbx/FbxSkeletalMeshImport.cpp
+++ b/Engine/Source/Editor/UnrealEd/Private/Fbx/FbxSkeletalMeshImport.cpp
@@ -1255,12 +1255,26 @@ USkeletalMesh* UnFbx::FFbxImporter::ImportSkeletalMesh(UObject* InParent, TArray
 		const bool bShouldComputeTangents = !ImportOptions->ShouldImportTangents() || !SkelMeshImportDataPtr->bHasTangents;
 
 		IMeshUtilities& MeshUtilities = FModuleManager::Get().LoadModuleChecked<IMeshUtilities>("MeshUtilities");
+		
+		TArray<FText> WarningMessages;
 		// Create actual rendering data.
-		if ( !MeshUtilities.BuildSkeletalMesh(ImportedResource->LODModels[0],SkeletalMesh->RefSkeleton,LODInfluences,LODWedges,LODFaces,LODPoints,LODPointToRawMap, ImportOptions->bKeepOverlappingVertices, bShouldComputeNormals, bShouldComputeTangents ) )
+		if ( !MeshUtilities.BuildSkeletalMesh(ImportedResource->LODModels[0],SkeletalMesh->RefSkeleton,LODInfluences,LODWedges,LODFaces,LODPoints,LODPointToRawMap, ImportOptions->bKeepOverlappingVertices, bShouldComputeNormals, bShouldComputeTangents, &WarningMessages ) )
 		{
+			for (auto Message : WarningMessages)
+			{
+				AddTokenizedErrorMessage(FTokenizedMessage::Create(EMessageSeverity::Error, Message));
+			}
+
 			SkeletalMesh->MarkPendingKill();
 			return NULL;
 		}
+		else if (WarningMessages.Num() > 0)
+		{
+			for (auto Message : WarningMessages)
+			{
+				AddTokenizedErrorMessage(FTokenizedMessage::Create(EMessageSeverity::Warning, Message));
+			}
+		}
 
 		// Presize the per-section shadow casting array with the number of sections in the imported LOD.
 		const int32 NumSections = LODModel.Sections.Num();
