added better warning for build skeletalmesh fails.

diff --git a/Engine/Source/Developer/MeshUtilities/Private/MeshUtilities.cpp b/Engine/Source/Developer/MeshUtilities/Private/MeshUtilities.cpp
--- a/Engine/Source/Developer/MeshUtilities/Private/MeshUtilities.cpp
+++ b/Engine/Source/Developer/MeshUtilities/Private/MeshUtilities.cpp
@@ -15,6 +15,8 @@
 #include "MaterialExportUtils.h"
 #include "Textures/TextureAtlas.h"
 
+#include "FbxErrors.h"
+
 //@todo - implement required vector intrinsics for other implementations
 #if PLATFORM_ENABLE_VECTORINTRINSICS
 #include "kDOP.h"
@@ -482,7 +484,7 @@ private:
 		const FStaticMeshLODGroup& LODGroup
 		) override;
 
-	virtual bool BuildSkeletalMesh( FStaticLODModel& LODModel, const FReferenceSkeleton& RefSkeleton, const TArray<FVertInfluence>& Influences, const TArray<FMeshWedge>& Wedges, const TArray<FMeshFace>& Faces, const TArray<FVector>& Points, const TArray<int32>& PointToOriginalMap, bool bKeepOverlappingVertices = false, bool bComputeNormals = true, bool bComputeTangents = true, TArray<FText> * OutWarningMessages = NULL);
+	virtual bool BuildSkeletalMesh( FStaticLODModel& LODModel, const FReferenceSkeleton& RefSkeleton, const TArray<FVertInfluence>& Influences, const TArray<FMeshWedge>& Wedges, const TArray<FMeshFace>& Faces, const TArray<FVector>& Points, const TArray<int32>& PointToOriginalMap, bool bKeepOverlappingVertices = false, bool bComputeNormals = true, bool bComputeTangents = true, TArray<FText> * OutWarningMessages = NULL, TArray<FName> * OutWarningNames = NULL);
 
 	virtual bool GenerateUVs(
 		FRawMesh& RawMesh,
@@ -2551,7 +2553,8 @@ bool FMeshUtilities::BuildStaticMesh(
 }
 
 
-bool FMeshUtilities::BuildSkeletalMesh( FStaticLODModel& LODModel, const FReferenceSkeleton& RefSkeleton, const TArray<FVertInfluence>& Influences, const TArray<FMeshWedge>& Wedges, const TArray<FMeshFace>& Faces, const TArray<FVector>& Points, const TArray<int32>& PointToOriginalMap, bool bKeepOverlappingVertices, bool bComputeNormals, bool bComputeTangents, TArray<FText> * OutWarningMessages )
+//@TODO: The OutMessages has to be a struct that contains FText/FName, or make it Token and add that as error. Needs re-work. Temporary workaround for now. 
+bool FMeshUtilities::BuildSkeletalMesh( FStaticLODModel& LODModel, const FReferenceSkeleton& RefSkeleton, const TArray<FVertInfluence>& Influences, const TArray<FMeshWedge>& Wedges, const TArray<FMeshFace>& Faces, const TArray<FVector>& Points, const TArray<int32>& PointToOriginalMap, bool bKeepOverlappingVertices, bool bComputeNormals, bool bComputeTangents, TArray<FText> * OutWarningMessages, TArray<FName> * OutWarningNames)
 {
 #if WITH_EDITORONLY_DATA
 	bool bTooManyVerts = false;
@@ -2637,6 +2640,10 @@ bool FMeshUtilities::BuildSkeletalMesh( FStaticLODModel& LODModel, const FRefere
 			if (OutWarningMessages)
 			{
 				OutWarningMessages->Add(FText::Format(FText::FromString("Missing influence on vert {0}. Weighting it to root."), FText::FromString(FString::FromInt(Wedges[WedgeIndex].iVertex))));
+				if (OutWarningNames)
+				{
+					OutWarningNames->Add(FFbxErrors::SkeletalMesh_VertMissingInfluences);
+				}
 			}
 		}
 	}
@@ -2986,6 +2993,10 @@ bool FMeshUtilities::BuildSkeletalMesh( FStaticLODModel& LODModel, const FRefere
 			if(OutWarningMessages)
 			{
 				OutWarningMessages->Add(BadSectionMessage);
+				if(OutWarningNames)
+				{
+					OutWarningNames->Add(FFbxErrors::SkeletalMesh_SectionWithNoTriangle);
+				}
 			}
 			else
 			{
@@ -3000,6 +3011,10 @@ bool FMeshUtilities::BuildSkeletalMesh( FStaticLODModel& LODModel, const FRefere
 			if(OutWarningMessages)
 			{
 				OutWarningMessages->Add(TooManyVertsMessage);
+				if(OutWarningNames)
+				{
+					OutWarningNames->Add(FFbxErrors::SkeletalMesh_TooManyVertices);
+				}
 			}
 			else
 			{
diff --git a/Engine/Source/Developer/MeshUtilities/Public/MeshUtilities.h b/Engine/Source/Developer/MeshUtilities/Public/MeshUtilities.h
--- a/Engine/Source/Developer/MeshUtilities/Public/MeshUtilities.h
+++ b/Engine/Source/Developer/MeshUtilities/Public/MeshUtilities.h
@@ -150,7 +150,8 @@ public:
 		bool bKeepOverlappingVertices = false,
 		bool bComputeNormals = true,
 		bool bComputeTangents = true,
-		TArray<FText> * OutWarningMessages = NULL
+		TArray<FText> * OutWarningMessages = NULL,
+		TArray<FName> * OutWarningNames = NULL
 		) = 0;
 
 
diff --git a/Engine/Source/Editor/UnrealEd/Private/Fbx/FbxSkeletalMeshImport.cpp b/Engine/Source/Editor/UnrealEd/Private/Fbx/FbxSkeletalMeshImport.cpp
--- a/Engine/Source/Editor/UnrealEd/Private/Fbx/FbxSkeletalMeshImport.cpp
+++ b/Engine/Source/Editor/UnrealEd/Private/Fbx/FbxSkeletalMeshImport.cpp
@@ -1259,12 +1259,17 @@ USkeletalMesh* UnFbx::FFbxImporter::ImportSkeletalMesh(UObject* InParent, TArray
 		IMeshUtilities& MeshUtilities = FModuleManager::Get().LoadModuleChecked<IMeshUtilities>("MeshUtilities");
 		
 		TArray<FText> WarningMessages;
+		TArray<FName> WarningNames;
 		// Create actual rendering data.
-		if ( !MeshUtilities.BuildSkeletalMesh(ImportedResource->LODModels[0],SkeletalMesh->RefSkeleton,LODInfluences,LODWedges,LODFaces,LODPoints,LODPointToRawMap, ImportOptions->bKeepOverlappingVertices, bShouldComputeNormals, bShouldComputeTangents, &WarningMessages ) )
+		if ( !MeshUtilities.BuildSkeletalMesh(ImportedResource->LODModels[0],SkeletalMesh->RefSkeleton,LODInfluences,LODWedges,LODFaces,LODPoints,LODPointToRawMap, ImportOptions->bKeepOverlappingVertices, bShouldComputeNormals, bShouldComputeTangents, &WarningMessages, &WarningNames ) )
 		{
-			for (auto Message : WarningMessages)
+			if (WarningNames.Num() == WarningMessages.Num())
 			{
-				AddTokenizedErrorMessage(FTokenizedMessage::Create(EMessageSeverity::Error, Message), FFbxErrors::SkeletalMesh_BuildingSkeletalMesh);
+				// temporary hack of message/names, should be one token or a struct
+				for(int32 MessageIdx = 0; MessageIdx<WarningMessages.Num(); ++MessageIdx)
+				{
+					AddTokenizedErrorMessage(FTokenizedMessage::Create(EMessageSeverity::Error, WarningMessages[MessageIdx]), WarningNames[MessageIdx]);
+				}
 			}
 
 			SkeletalMesh->MarkPendingKill();
@@ -1272,9 +1277,14 @@ USkeletalMesh* UnFbx::FFbxImporter::ImportSkeletalMesh(UObject* InParent, TArray
 		}
 		else if (WarningMessages.Num() > 0)
 		{
-			for (auto Message : WarningMessages)
+			// temporary hack of message/names, should be one token or a struct
+			if(WarningNames.Num() == WarningMessages.Num())
 			{
-				AddTokenizedErrorMessage(FTokenizedMessage::Create(EMessageSeverity::Warning, Message), FFbxErrors::SkeletalMesh_BuildingSkeletalMesh);
+				// temporary hack of message/names, should be one token or a struct
+				for(int32 MessageIdx = 0; MessageIdx<WarningMessages.Num(); ++MessageIdx)
+				{
+					AddTokenizedErrorMessage(FTokenizedMessage::Create(EMessageSeverity::Warning, WarningMessages[MessageIdx]), WarningNames[MessageIdx]);
+				}
 			}
 		}
 
diff --git a/Engine/Source/Runtime/Core/Private/Misc/FbxErrors.cpp b/Engine/Source/Runtime/Core/Private/Misc/FbxErrors.cpp
--- a/Engine/Source/Runtime/Core/Private/Misc/FbxErrors.cpp
+++ b/Engine/Source/Runtime/Core/Private/Misc/FbxErrors.cpp
@@ -80,7 +80,11 @@ FName FFbxErrors::SkeletalMesh_MultipleRoots(TEXT("SkeletalMesh_MultipleRoots"))
 
 FName FFbxErrors::SkeletalMesh_BonesAreMissingFromBindPose(TEXT("SkeletalMesh_BonesAreMissingFromBindPose"));
 
-FName FFbxErrors::SkeletalMesh_BuildingSkeletalMesh(TEXT("SkeletalMesh_BuildingSkeletalMesh"));
+FName FFbxErrors::SkeletalMesh_VertMissingInfluences(TEXT("SkeletalMesh_VertMissingInfluences"));
+
+FName FFbxErrors::SkeletalMesh_SectionWithNoTriangle(TEXT("SkeletalMesh_SectionWithNoTriangle"));
+
+FName FFbxErrors::SkeletalMesh_TooManyVertices(TEXT("SkeletalMesh_TooManyVertices"));
 
 FName FFbxErrors::SkeletalMesh_FailedToCreatePhyscisAsset(TEXT("SkeletalMesh_FailedToCreatePhyscisAsset"));
 
diff --git a/Engine/Source/Runtime/Core/Public/Misc/FbxErrors.h b/Engine/Source/Runtime/Core/Public/Misc/FbxErrors.h
--- a/Engine/Source/Runtime/Core/Public/Misc/FbxErrors.h
+++ b/Engine/Source/Runtime/Core/Public/Misc/FbxErrors.h
@@ -87,7 +87,11 @@ struct CORE_API FFbxErrors
 
 	static FName SkeletalMesh_BonesAreMissingFromBindPose;
 
-	static FName SkeletalMesh_BuildingSkeletalMesh;
+	static FName SkeletalMesh_VertMissingInfluences;
+
+	static FName SkeletalMesh_SectionWithNoTriangle;
+
+	static FName SkeletalMesh_TooManyVertices;
 
 	static FName SkeletalMesh_FailedToCreatePhyscisAsset;
 
