Make sure the index is within range - reported by licensee : https://udn.unrealengine.com/questions/175299/bug-in-meshutilitiescpp.html

diff --git a/Engine/Source/Developer/MeshUtilities/Private/MeshUtilities.cpp b/Engine/Source/Developer/MeshUtilities/Private/MeshUtilities.cpp
--- a/Engine/Source/Developer/MeshUtilities/Private/MeshUtilities.cpp
+++ b/Engine/Source/Developer/MeshUtilities/Private/MeshUtilities.cpp
@@ -2154,7 +2154,7 @@ static void BuildStaticMeshVertexAndIndexBuffers(
 		}
 
 		// Put the indices in the material index buffer.
-		int32 SectionIndex = FMath::Clamp(RawMesh.FaceMaterialIndices[FaceIndex], 0, OutPerSectionIndices.Num());
+		int32 SectionIndex = FMath::Clamp(RawMesh.FaceMaterialIndices[FaceIndex], 0, OutPerSectionIndices.Num()-1);
 		TArray<uint32>& SectionIndices = OutPerSectionIndices[SectionIndex];
 		for (int32 CornerIndex = 0; CornerIndex < 3; CornerIndex++)
 		{
