Auto lightmap UVs.
Replaces D3DMeshUtils.LayoutUVs.

Repacks existing UV charts. Will join charts in cases of tiling.

diff --git a/Engine/Source/Developer/MeshUtilities/Private/Allocator2D.cpp b/Engine/Source/Developer/MeshUtilities/Private/Allocator2D.cpp
new file mode 100644
--- /dev/null
+++ b/Engine/Source/Developer/MeshUtilities/Private/Allocator2D.cpp
@@ -0,0 +1,84 @@
+// Copyright 1998-2014 Epic Games, Inc. All Rights Reserved.
+
+#include "MeshUtilitiesPrivate.h"
+#include "Allocator2D.h"
+
+FAllocator2D::FAllocator2D( uint32 InWidth, uint32 InHeight )
+	: Width( InWidth )
+	, Height( InHeight )
+{
+	Pitch = ( Width + 63 ) / 64;
+	// alloc +1 to avoid buffer overrun
+	Bits = (uint64*)FMemory::Malloc( (Pitch * Height + 1) * sizeof(uint64) );
+}
+
+FAllocator2D::~FAllocator2D()
+{
+	FMemory::Free( Bits );
+}
+
+void FAllocator2D::Clear()
+{
+	FMemory::Memzero( Bits, (Pitch * Height + 1) * sizeof(uint64) );
+}
+
+bool FAllocator2D::Find( FRect& Rect )
+{
+	FRect TestRect = Rect;
+	for( TestRect.X = 0; TestRect.X <= Width - TestRect.W; TestRect.X++ )
+	{
+		for( TestRect.Y = 0; TestRect.Y <= Height - TestRect.H; TestRect.Y++ )
+		{
+			if( Test( TestRect ) )
+			{
+				Rect = TestRect;
+				return true;
+			}
+		}
+	}
+
+	return false;
+}
+
+bool FAllocator2D::Find( FRect& Rect, const FAllocator2D& Other )
+{
+	FRect TestRect = Rect;
+	for( TestRect.X = 0; TestRect.X <= Width - TestRect.W; TestRect.X++ )
+	{
+		for( TestRect.Y = 0; TestRect.Y <= Height - TestRect.H; TestRect.Y++ )
+		{
+			if( Test( TestRect, Other ) )
+			{
+				Rect = TestRect;
+				return true;
+			}
+		}
+	}
+	
+	return false;
+}
+
+void FAllocator2D::Alloc( FRect Rect )
+{
+	for( uint32 y = Rect.Y; y < Rect.Y + Rect.H; y++ )
+	{
+		for( uint32 x = Rect.X; x < Rect.X + Rect.W; x++ )
+		{
+			SetBit( x, y );
+		}
+	}
+}
+
+void FAllocator2D::Alloc( FRect Rect, const FAllocator2D& Other )
+{
+	for( uint32 y = 0; y < Rect.H; y++ )
+	{
+		for( uint32 x = 0; x < Rect.W; x++ )
+		{
+			if( Other.GetBit( x, y ) )
+			{
+				SetBit( x + Rect.X, y + Rect.Y );
+			}
+		}
+	}
+}
\ No newline at end of file
diff --git a/Engine/Source/Developer/MeshUtilities/Private/Allocator2D.h b/Engine/Source/Developer/MeshUtilities/Private/Allocator2D.h
new file mode 100644
--- /dev/null
+++ b/Engine/Source/Developer/MeshUtilities/Private/Allocator2D.h
@@ -0,0 +1,110 @@
+// Copyright 1998-2014 Epic Games, Inc. All Rights Reserved.
+
+#pragma once
+
+class FAllocator2D
+{
+public:
+	struct FRect
+	{
+		uint32 X;
+		uint32 Y;
+		uint32 W;
+		uint32 H;
+	};
+
+public:
+				FAllocator2D( uint32 Width, uint32 Height );
+				~FAllocator2D();
+	
+	// Must clear before using
+	void		Clear();
+
+	bool		Find( FRect& Rect );
+	bool		Test( FRect Rect );
+	void		Alloc( FRect Rect );
+
+	bool		Find( FRect& Rect, const FAllocator2D& Other );
+	bool		Test( FRect Rect, const FAllocator2D& Other );
+	void		Alloc( FRect Rect, const FAllocator2D& Other );
+	
+	uint64		GetBit( uint32 x, uint32 y ) const;
+	void		SetBit( uint32 x, uint32 y );
+
+private:
+	uint64*		Bits;
+
+	uint32		Width;
+	uint32		Height;
+	uint32		Pitch;
+};
+
+// Returns non-zero if set
+FORCEINLINE uint64 FAllocator2D::GetBit( uint32 x, uint32 y ) const
+{
+	return Bits[ (x >> 6) + y * Pitch ] & ( 1ull << ( x & 63 ) );
+}
+
+FORCEINLINE void FAllocator2D::SetBit( uint32 x, uint32 y )
+{
+	Bits[ (x >> 6) + y * Pitch ] |= ( 1ull << ( x & 63 ) );
+}
+
+inline bool FAllocator2D::Test( FRect Rect )
+{
+	for( uint32 y = Rect.Y; y < Rect.Y + Rect.H; y++ )
+	{
+		for( uint32 x = Rect.X; x < Rect.X + Rect.W; x++ )
+		{
+			if( GetBit( x, y ) )
+			{
+				return false;
+			}
+		}
+	}
+	
+	return true;
+}
+
+inline bool FAllocator2D::Test( FRect Rect, const FAllocator2D& Other )
+{
+	const uint32 LowShift = Rect.X & 63;
+	const uint32 HighShift = 64 - LowShift;
+
+	for( uint32 y = 0; y < Rect.H; y++ )
+	{
+#if 1
+		uint32 ThisIndex = (Rect.X >> 6) + (y + Rect.Y) * Pitch;
+		uint32 OtherIndex = y * Pitch;
+
+		// Test a uint64 at a time
+		for( uint32 x = 0; x < Rect.W; x += 64 )
+		{
+			// no need to zero out HighInt on wrap around because Other will always be zero outside Rect.
+			uint64 LowInt  = Bits[ ThisIndex ];
+			uint64 HighInt = Bits[ ThisIndex + 1 ];
+
+			uint64 ThisInt = (HighInt << HighShift) | (LowInt >> LowShift);
+			uint64 OtherInt = Other.Bits[ OtherIndex ];
+
+			if( ThisInt & OtherInt )
+			{
+				return false;
+			}
+
+			ThisIndex++;
+			OtherIndex++;
+		}
+#else
+		for( uint32 x = 0; x < Rect.W; x++ )
+		{
+			if( Other.GetBit( x, y ) && GetBit( x + Rect.X, y + Rect.Y ) )
+			{
+				return false;
+			}
+		}
+#endif
+	}
+	
+	return true;
+}
\ No newline at end of file
diff --git a/Engine/Source/Developer/MeshUtilities/Private/DisjointSet.h b/Engine/Source/Developer/MeshUtilities/Private/DisjointSet.h
new file mode 100644
--- /dev/null
+++ b/Engine/Source/Developer/MeshUtilities/Private/DisjointSet.h
@@ -0,0 +1,80 @@
+// Copyright 1998-2014 Epic Games, Inc. All Rights Reserved.
+
+#pragma once
+
+class FDisjointSet
+{
+public:
+			FDisjointSet( const uint32 Size );
+	
+	void	Union( uint32 x, uint32 y );
+	uint32	Find( uint32 i );
+
+	uint32	operator[]( uint32 i ) const	{ return Parents[i]; }
+
+private:
+	TArray<	uint32 >	Parents;
+};
+
+FDisjointSet::FDisjointSet( const uint32 Size )
+{
+	Parents.Init( Size );
+	for( uint32 i = 0; i < Size; i++ )
+	{
+		Parents[i] = i;
+	}
+}
+
+// Union with splicing
+void FDisjointSet::Union( uint32 x, uint32 y )
+{
+	while( Parents[x] != Parents[y] )
+	{
+		// Pick larger
+		if( Parents[x] < Parents[y] )
+		{
+			uint32 p = Parents[x];
+			Parents[x] = Parents[y];
+			if( x == p )
+			{
+				return;
+			}
+			x = p;
+		}
+		else
+		{
+			uint32 p = Parents[y];
+			Parents[y] = Parents[x];
+			if( y == p )
+			{
+				return;
+			}
+			y = p;
+		}
+	}
+}
+
+// Find with path compression
+uint32 FDisjointSet::Find( uint32 i )
+{
+	// Find root
+	uint32 Start = i;
+	uint32 Root = Parents[i];
+	while( Root != i )
+	{
+		i = Root;
+		Root = Parents[i];
+	}
+
+	// Point all nodes on path to root
+	i = Start;
+	uint32 Parent = Parents[i];
+	while( Parent != Root )
+	{
+		Parents[i] = Root;
+		i = Parent;
+		Parent = Parents[i];
+	}
+
+	return Root;
+}
\ No newline at end of file
diff --git a/Engine/Source/Developer/MeshUtilities/Private/LayoutUV.cpp b/Engine/Source/Developer/MeshUtilities/Private/LayoutUV.cpp
new file mode 100644
--- /dev/null
+++ b/Engine/Source/Developer/MeshUtilities/Private/LayoutUV.cpp
@@ -0,0 +1,1021 @@
+// Copyright 1998-2014 Epic Games, Inc. All Rights Reserved.
+
+#include "MeshUtilitiesPrivate.h"
+#include "LayoutUV.h"
+#include "DisjointSet.h"
+
+DEFINE_LOG_CATEGORY_STATIC(LogLayoutUV, Log, All);
+
+#define CHART_JOINING	1
+
+FLayoutUV::FLayoutUV( FRawMesh* InMesh, uint32 InSrcChannel, uint32 InDstChannel, uint32 InTextureResolution )
+	: RawMesh( InMesh )
+	, SrcChannel( InSrcChannel )
+	, DstChannel( InDstChannel )
+	, TextureResolution( InTextureResolution )
+	, TotalUVArea( 0.0f )
+	, LayoutRaster( TextureResolution, TextureResolution )
+	, ChartRaster( TextureResolution, TextureResolution )
+	, ChartShader( &ChartRaster )
+{}
+
+void FindOverlappingCorners( TMultiMap<int32,int32>& OutOverlappingCorners, FRawMesh const& RawMesh, float ComparisonThreshold );
+
+void FLayoutUV::FindCharts()
+{
+	double Begin = FPlatformTime::Seconds();
+
+	TMultiMap<int32,int32> OverlappingCorners;
+
+	// Find overlapping corners to accelerate adjacency.
+	FindOverlappingCorners( OverlappingCorners, *RawMesh, THRESH_POINTS_ARE_SAME );
+
+	uint32 NumIndexes = RawMesh->WedgeIndices.Num();
+	uint32 NumTris = NumIndexes / 3;
+
+	TArray< int32 > TranslatedMatches;
+	TranslatedMatches.Init( NumIndexes );
+	TexCoords.Init( NumIndexes );
+	for( uint32 i = 0; i < NumIndexes; i++ )
+	{
+		TranslatedMatches[i] = -1;
+		TexCoords[i] = RawMesh->WedgeTexCoords[ SrcChannel ][i];
+	}
+
+	// Build disjoint set
+	FDisjointSet DisjointSet( NumTris );
+	for( uint32 i = 0; i < NumIndexes; i++ )
+	{
+		for( auto It = OverlappingCorners.CreateConstKeyIterator(i); It; ++It )
+		{
+			uint32 j = It.Value();
+
+			if( j > i )
+			{
+				const uint32 TriI = i/3;
+				const uint32 TriJ = j/3;
+
+				bool bUnion = false;
+
+#if CHART_JOINING
+				bool bPositionMatch = PositionsMatch( i, j );
+				if( bPositionMatch )
+				{
+					uint32 i1 = 3 * TriI + (i + 1) % 3;
+					uint32 i2 = 3 * TriI + (i + 2) % 3;
+					uint32 j1 = 3 * TriJ + (j + 1) % 3;
+					uint32 j2 = 3 * TriJ + (j + 2) % 3;
+
+					bool bEdgeMatch21 = PositionsMatch( i2, j1 );
+					bool bEdgeMatch12 = PositionsMatch( i1, j2 );
+					if( bEdgeMatch21 || bEdgeMatch12 )
+					{
+						uint32 ie = bEdgeMatch21 ? i2 : i1;
+						uint32 je = bEdgeMatch21 ? j1 : j2;
+							
+						bool bUVMatch = UVsMatch( i, j ) && UVsMatch( ie, je );
+						bool bUVWindingMatch = TriangleUVArea( TriI ) * TriangleUVArea( TriJ ) >= 0.0f;
+						if( bUVMatch && bUVWindingMatch )
+						{
+							bUnion = true;
+						}
+						else if( NormalsMatch( i, j ) && NormalsMatch( ie, je ) )
+						{
+							// Chart edge
+							FVector2D EdgeUVi = TexCoords[ie] - TexCoords[i];
+							FVector2D EdgeUVj = TexCoords[je] - TexCoords[j];
+							
+							// Would these edges match if the charts were translated
+							bool bTranslatedUVMatch = ( EdgeUVi - EdgeUVj ).IsNearlyZero( THRESH_UVS_ARE_SAME );
+							if( bTranslatedUVMatch )
+							{
+								// Note: may be mirrored
+								
+								// TODO should these be restricted to axis aligned edges?
+								uint32 EdgeI = bEdgeMatch21 ? i2 : i;
+								uint32 EdgeJ = bEdgeMatch21 ? j : j2;
+
+								// Only allow one match per edge
+								if( TranslatedMatches[ EdgeI ] < 0 &&
+									TranslatedMatches[ EdgeJ ] < 0 )
+								{
+									TranslatedMatches[ EdgeI ] = EdgeJ;
+									TranslatedMatches[ EdgeJ ] = EdgeI;
+								}
+							}
+						}
+					}
+				}
+#else
+				if( VertsMatch( i, j ) )
+				{
+					// Edge must match as well (same winding)
+					if( VertsMatch( 3 * TriI + (i - 1) % 3, 3 * TriJ + (j + 1) % 3 ) ||
+						VertsMatch( 3 * TriI + (i + 1) % 3, 3 * TriJ + (j - 1) % 3 ) )
+					{
+						// Check for UV winding match too
+						if( TriangleUVArea( TriI ) * TriangleUVArea( TriJ ) >= 0.0f )
+						{
+							bUnion = true;
+						}
+					}
+				}
+#endif
+
+				if( bUnion )
+				{
+					// TODO solve spiral case by checking sets for UV overlap
+					DisjointSet.Union( TriI, TriJ );
+				}
+			}
+		}
+	}
+
+	// Sort tris by chart
+	SortedTris.Init( NumTris );
+	for( uint32 i = 0; i < NumTris; i++ )
+	{
+		// Flatten disjoint set path
+		DisjointSet.Find(i);
+		SortedTris[i] = i;
+	}
+
+	struct FCompareTris
+	{
+		FDisjointSet* DisjointSet;
+
+		FCompareTris( FDisjointSet* InDisjointSet )
+		: DisjointSet( InDisjointSet )
+		{}
+
+		FORCEINLINE bool operator()( uint32 A, uint32 B ) const
+		{
+			return (*DisjointSet)[A] < (*DisjointSet)[B];
+		}
+	};
+	SortedTris.Sort( FCompareTris( &DisjointSet ) );
+
+	TMap< uint32, int32 > DisjointSetToChartMap;
+
+	// Build Charts
+	for( uint32 Tri = 0; Tri < NumTris; )
+	{
+		int32 i = Charts.AddUninitialized();
+		FMeshChart& Chart = Charts[i];
+		
+		Chart.MinUV = FVector2D( FLT_MAX, FLT_MAX );
+		Chart.MaxUV = FVector2D( -FLT_MAX, -FLT_MAX );
+		Chart.UVArea = 0.0f;
+		Chart.WorldScale = FVector2D::ZeroVector;
+		FMemory::MemSet( Chart.Join, 0xff );
+
+		Chart.FirstTri = Tri;
+
+		uint32 ChartID = DisjointSet[ SortedTris[ Tri ] ];
+		DisjointSetToChartMap.Add( ChartID, i );
+
+		for( ; Tri < NumTris && DisjointSet[ SortedTris[ Tri ] ] == ChartID; Tri++ )
+		{
+			// Calculate chart bounds
+			FVector		Positions[3];
+			FVector2D	UVs[3];
+			for( int k = 0; k < 3; k++ )
+			{
+				uint32 Index = 3 * SortedTris[ Tri ] + k;
+
+				Positions[k] = RawMesh->GetWedgePosition( Index );
+				UVs[k] = TexCoords[ Index ];
+
+				Chart.MinUV.X = FMath::Min( Chart.MinUV.X, UVs[k].X );
+				Chart.MinUV.Y = FMath::Min( Chart.MinUV.Y, UVs[k].Y );
+				Chart.MaxUV.X = FMath::Max( Chart.MaxUV.X, UVs[k].X );
+				Chart.MaxUV.Y = FMath::Max( Chart.MaxUV.Y, UVs[k].Y );
+			}
+
+			FVector Edge1 = Positions[1] - Positions[0];
+			FVector Edge2 = Positions[2] - Positions[0];
+			float Area = 0.5f * ( Edge1 ^ Edge2 ).Size();
+
+			FVector2D EdgeUV1 = UVs[1] - UVs[0];
+			FVector2D EdgeUV2 = UVs[2] - UVs[0];
+			float UVArea = 0.5f * FMath::Abs( EdgeUV1.X * EdgeUV2.Y - EdgeUV1.Y * EdgeUV2.X );
+
+			FVector2D UVLength;
+			UVLength.X = ( EdgeUV2.Y * Edge1 - EdgeUV1.Y * Edge2 ).Size();
+			UVLength.Y = (-EdgeUV2.X * Edge1 + EdgeUV1.X * Edge2 ).Size();
+			
+			Chart.WorldScale += UVLength;
+			Chart.UVArea += UVArea;
+		}
+		
+		Chart.LastTri = Tri;
+
+#if !CHART_JOINING
+		if( Chart.UVArea > 1e-4f )
+		{
+			Chart.WorldScale /= Chart.UVArea;
+		}
+		else
+		{
+			Chart.WorldScale = FVector2D::ZeroVector;
+		}
+
+		//Chart.WorldScale.Set(1,1);
+
+		TotalUVArea += Chart.UVArea * Chart.WorldScale.X * Chart.WorldScale.Y;
+#endif
+	}
+
+#if CHART_JOINING
+	for( int32 i = 0; i < Charts.Num(); i++ )
+	{
+		FMeshChart& Chart = Charts[i];
+
+		for( uint32 Tri = Chart.FirstTri; Tri < Chart.LastTri; Tri++ )
+		{
+			for( int k = 0; k < 3; k++ )
+			{
+				uint32 Index = 3 * SortedTris[ Tri ] + k;
+
+				if( TranslatedMatches[ Index ] >= 0 )
+				{
+					checkSlow( TranslatedMatches[ TranslatedMatches[ Index ] ] == Index );
+
+					uint32 V0i = Index;
+					uint32 V0j = TranslatedMatches[ Index ];
+
+					uint32 TriI = V0i / 3;
+					uint32 TriJ = V0j / 3;
+					
+					if( TriJ <= TriI )
+					{
+						// Only need to consider one direction
+						continue;
+					}
+					
+					uint32 V1i = 3 * TriI + (V0i + 1) % 3;
+					uint32 V1j = 3 * TriJ + (V0j + 1) % 3;
+
+					int32 ChartI = i;
+					int32 ChartJ = DisjointSetToChartMap[ DisjointSet[ TriJ ] ];
+
+					FVector2D UV0i = TexCoords[ V0i ];
+					FVector2D UV1i = TexCoords[ V1i ];
+					FVector2D UV0j = TexCoords[ V0j ];
+					FVector2D UV1j = TexCoords[ V1j ];
+
+					FVector2D EdgeUVi = UV1i - UV0i;
+					FVector2D EdgeUVj = UV1j - UV0j;
+
+					bool bMirrored = TriangleUVArea( TriI ) * TriangleUVArea( TriJ ) < 0.0f;
+					
+					FVector2D EdgeOffset0 = UV0i - UV1j;
+					FVector2D EdgeOffset1 = UV1i - UV0j;
+					checkSlow( ( EdgeOffset0 - EdgeOffset1 ).IsNearlyZero( THRESH_UVS_ARE_SAME ) );
+
+					FVector2D Translation = EdgeOffset0;
+
+					FMeshChart& ChartA = Charts[ ChartI ];
+					FMeshChart& ChartB = Charts[ ChartJ ];
+
+					for( uint32 Side = 0; Side < 4; Side++ )
+					{
+						// Join[] = { left, right, bottom, top }
+
+						// FIXME
+						if( bMirrored )
+							continue;
+
+						uint32 Sign = Side & 1;
+						uint32 Axis = Side >> 1;
+
+						bool bAxisAligned = FMath::Abs( EdgeUVi[ Axis ] ) < THRESH_UVS_ARE_SAME;
+						bool bBorderA = FMath::Abs( UV0i[ Axis ] - ( Sign ^ 0 ? Chart.MaxUV[ Axis ] : Chart.MinUV[ Axis ] ) ) < THRESH_UVS_ARE_SAME;
+						bool bBorderB = FMath::Abs( UV0j[ Axis ] - ( Sign ^ 1 ? Chart.MaxUV[ Axis ] : Chart.MinUV[ Axis ] ) ) < THRESH_UVS_ARE_SAME;
+
+						// FIXME mirrored
+						if( !bAxisAligned || !bBorderA || !bBorderB )
+						{
+							// Edges weren't on matching rectangle borders
+							continue;
+						}
+
+						FVector2D CenterA = 0.5f * ( ChartA.MinUV + ChartA.MaxUV );
+						FVector2D CenterB = 0.5f * ( ChartB.MinUV + ChartB.MaxUV );
+
+						FVector2D ExtentA = 0.5f * ( ChartA.MaxUV - ChartA.MinUV );
+						FVector2D ExtentB = 0.5f * ( ChartB.MaxUV - ChartB.MinUV );
+
+						// FIXME mirrored
+						CenterB += Translation;
+
+						FVector2D CenterDiff = CenterA - CenterB;
+						FVector2D ExtentDiff = ExtentA - ExtentB;
+						FVector2D Separation = ExtentA + ExtentB + CenterDiff * ( Sign ? 1.0f : -1.0f );
+
+						bool bCenterMatch = FMath::Abs( CenterDiff[ Axis ^ 1 ] ) < THRESH_UVS_ARE_SAME;
+						bool bExtentMatch = FMath::Abs( ExtentDiff[ Axis ^ 1 ] ) < THRESH_UVS_ARE_SAME;
+						bool bSeparate    = FMath::Abs( Separation[ Axis ^ 0 ] ) < THRESH_UVS_ARE_SAME;
+	
+						if( !bCenterMatch || !bExtentMatch || !bSeparate )
+						{
+							// Rectangles don't match up after translation
+							continue;
+						}
+
+						// Found a valid edge join
+						ChartA.Join[ Side ^ 0 ] = ChartJ;
+						ChartB.Join[ Side ^ 1 ] = ChartI;
+						break;
+					}
+				}
+			}
+		}
+	}
+
+	TArray< uint32 > JoinedSortedTris;
+	JoinedSortedTris.Reserve( NumTris );
+
+	// Detect loops
+	for( uint32 Axis = 0; Axis < 2; Axis++ )
+	{
+		uint32 Side = Axis << 1;
+
+		for( int32 i = 0; i < Charts.Num(); i++ )
+		{
+			int32 j = Charts[i].Join[ Side ^ 1 ];
+			while( j != -1 )
+			{
+				int32 Next = Charts[j].Join[ Side ^ 1 ];
+				if( Next == i )
+				{
+					// Break loop
+					Charts[i].Join[ Side ^ 0 ] = -1;
+					Charts[j].Join[ Side ^ 1 ] = -1;
+					break;
+				}
+				j = Next;
+			}
+		}
+	}
+
+	// Join rows first, then columns
+	for( uint32 Axis = 0; Axis < 2; Axis++ )
+	{
+		for( int32 i = 0; i < Charts.Num(); i++ )
+		{
+			FMeshChart& ChartA = Charts[i];
+
+			if( ChartA.FirstTri == ChartA.LastTri )
+			{
+				// Empty chart
+				continue;
+			}
+			
+			for( uint32 Side = 0; Side < 4; Side++ )
+			{
+				if( ChartA.Join[ Side ] != -1 )
+				{
+					FMeshChart& ChartB = Charts[ ChartA.Join[ Side ] ];
+
+					check( ChartB.Join[ Side ^ 1 ] == i );
+					check( ChartB.FirstTri != ChartB.LastTri );
+				}
+			}
+		}
+
+		uint32 NumTris = 0;
+		for( int32 i = 0; i < Charts.Num(); i++ )
+		{
+			FMeshChart& Chart = Charts[i];
+			NumTris += Chart.LastTri - Chart.FirstTri;
+		}
+		check( NumTris == SortedTris.Num() );
+
+		NumTris = 0;
+		for( int32 i = 0; i < Charts.Num(); i++ )
+		{
+			FMeshChart& ChartA = Charts[i];
+
+			if( ChartA.FirstTri == ChartA.LastTri )
+			{
+				// Empty chart
+				continue;
+			}
+
+			uint32 Side = Axis << 1;
+
+			// Find start (left, bottom)
+			if( ChartA.Join[ Side ^ 0 ] == -1 )
+			{
+				// Add original tris
+				NumTris += ChartA.LastTri - ChartA.FirstTri;
+
+				// Continue joining until no more to the (right, top)
+				int32 Next = ChartA.Join[ Side ^ 1 ];
+				while( Next != -1 )
+				{
+					FMeshChart& ChartB = Charts[ Next ];
+
+					NumTris += ChartB.LastTri - ChartB.FirstTri;
+					Next = ChartB.Join[ Side ^ 1 ];
+				}
+			}
+		}
+		check( NumTris == SortedTris.Num() );
+
+#if 1
+		NumTris = 0;
+		for( int32 i = 0; i < Charts.Num(); i++ )
+		{
+			FMeshChart& ChartA = Charts[i];
+
+			if( ChartA.FirstTri == ChartA.LastTri )
+			{
+				// Empty chart
+				continue;
+			}
+
+			// Join[] = { left, right, bottom, top }
+
+			uint32 Side = Axis << 1;
+
+			// Find start (left, bottom)
+			if( ChartA.Join[ Side ^ 0 ] == -1 )
+			{
+				uint32 FirstTri = JoinedSortedTris.Num();
+			
+				// Add original tris
+				for( uint32 Tri = ChartA.FirstTri; Tri < ChartA.LastTri; Tri++ )
+				{
+					JoinedSortedTris.Add( SortedTris[ Tri ] );
+				}
+				NumTris += ChartA.LastTri - ChartA.FirstTri;
+
+				// Continue joining until no more to the (right, top)
+				while( ChartA.Join[ Side ^ 1 ] != -1 )
+				{
+					FMeshChart& ChartB = Charts[ ChartA.Join[ Side ^ 1 ] ];
+
+					check( ChartB.FirstTri != ChartB.LastTri );
+
+					FVector2D Translation = ChartA.MinUV - ChartB.MinUV;
+					Translation[ Axis ] += ChartA.MaxUV[ Axis ] - ChartA.MinUV[ Axis ];
+
+					for( uint32 Tri = ChartB.FirstTri; Tri < ChartB.LastTri; Tri++ )
+					{
+						JoinedSortedTris.Add( SortedTris[ Tri ] );
+						for( int k = 0; k < 3; k++ )
+						{
+							TexCoords[ 3 * SortedTris[ Tri ] + k ] += Translation;
+						}
+					}
+					NumTris += ChartB.LastTri - ChartB.FirstTri;
+
+					ChartA.Join[ Side ^ 1 ] = ChartB.Join[ Side ^ 1 ];
+					ChartA.MaxUV[ Axis ] += ChartB.MaxUV[ Axis ] - ChartB.MinUV[ Axis ];
+					ChartA.WorldScale += ChartB.WorldScale;
+					ChartA.UVArea += ChartB.UVArea;
+
+					ChartB.FirstTri = 0;
+					ChartB.LastTri = 0;
+					ChartB.UVArea = 0.0f;
+
+					DisconnectChart( ChartB, Side ^ 2 );
+					DisconnectChart( ChartB, Side ^ 3 );
+				}
+
+				ChartA.FirstTri = FirstTri;
+				ChartA.LastTri = JoinedSortedTris.Num();
+			}
+			else
+			{
+				// Make sure a starting chart could connect to this
+				FMeshChart& ChartB = Charts[ ChartA.Join[ Side ^ 0 ] ];
+				check( ChartB.Join[ Side ^ 1 ] == i );
+				check( ChartB.FirstTri != ChartB.LastTri );
+			}
+		}
+		check( NumTris == SortedTris.Num() );
+
+		check( SortedTris.Num() == JoinedSortedTris.Num() );
+		Exchange( SortedTris, JoinedSortedTris );
+		JoinedSortedTris.Reset();
+#endif
+	}
+	
+	// Clean out empty charts
+	for( int32 i = 0; i < Charts.Num(); i++ )
+	{
+		while( i < Charts.Num() && Charts[i].FirstTri == Charts[i].LastTri )
+		{
+			Charts.RemoveAtSwap(i);
+		}
+	}
+
+	for( int32 i = 0; i < Charts.Num(); i++ )
+	{
+		FMeshChart& Chart = Charts[i];
+
+		if( Chart.UVArea > 1e-4f )
+		{
+			Chart.WorldScale /= Chart.UVArea;
+		}
+		else
+		{
+			Chart.WorldScale = FVector2D::ZeroVector;
+		}
+
+		TotalUVArea += Chart.UVArea * Chart.WorldScale.X * Chart.WorldScale.Y;
+	}
+#endif
+
+	double End = FPlatformTime::Seconds();
+
+	UE_LOG(LogLayoutUV, Display, TEXT("FindCharts: %s"), *FPlatformTime::PrettyTime(End - Begin) );
+}
+
+bool FLayoutUV::FindBestPacking()
+{
+	if( (uint32)Charts.Num() > TextureResolution * TextureResolution )
+	{
+		// More charts than texels
+		return false;
+	}
+	
+	const float LinearSearchStart = 0.5f;
+	const float LinearSearchStep = 0.5f;
+	const int32 BinarySearchSteps = 6;
+	
+	float UVScaleFail = TextureResolution * FMath::Sqrt( 1.0f / TotalUVArea );
+	float UVScalePass = TextureResolution * FMath::Sqrt( LinearSearchStart / TotalUVArea );
+
+	// Linear search for first fit
+	while(1)
+	{
+		ScaleCharts( UVScalePass );
+
+		bool bFit = PackCharts();
+		if( bFit )
+		{
+			break;
+		}
+
+		UVScaleFail = UVScalePass;
+		UVScalePass *= LinearSearchStep;
+	}
+
+	// Binary search for best fit
+	for( int32 i = 0; i < BinarySearchSteps; i++ )
+	{
+		float UVScale = 0.5f * ( UVScaleFail + UVScalePass );
+		ScaleCharts( UVScale );
+
+		bool bFit = PackCharts();
+		if( bFit )
+		{
+			UVScalePass = UVScale;
+		}
+		else
+		{
+			UVScaleFail = UVScale;
+		}
+	}
+
+	// TODO store packing scale/bias separate so this isn't necessary
+	ScaleCharts( UVScalePass );
+	PackCharts();
+
+	return true;
+}
+
+void FLayoutUV::ScaleCharts( float UVScale )
+{
+	for( int32 i = 0; i < Charts.Num(); i++ )
+	{
+		FMeshChart& Chart = Charts[i];
+		Chart.UVScale = Chart.WorldScale * UVScale;
+	}
+	
+	// Scale charts such that they all fit and roughly total the same area as before
+#if 1
+	float UniformScale = 1.0f;
+	while(1)
+	{
+		uint32 NumMaxedOut = 0;
+		float ScaledUVArea = 0.0f;
+		for( int32 i = 0; i < Charts.Num(); i++ )
+		{
+			FMeshChart& Chart = Charts[i];
+
+			FVector2D ChartSize	= Chart.MaxUV - Chart.MinUV;
+			FVector2D ChartSizeScaled = ChartSize * Chart.UVScale * UniformScale;
+
+			const float MaxChartEdge = TextureResolution - 1.0f;
+			const float LongestChartEdge = FMath::Max( ChartSizeScaled.X, ChartSizeScaled.Y );
+
+			const float Epsilon = 0.01f;
+			if( LongestChartEdge + Epsilon > MaxChartEdge )
+			{
+				// Rescale oversized charts to fit
+				Chart.UVScale.X = MaxChartEdge / FMath::Max( ChartSize.X, ChartSize.Y );
+				Chart.UVScale.Y = MaxChartEdge / FMath::Max( ChartSize.X, ChartSize.Y );
+				NumMaxedOut++;
+			}
+			else
+			{
+				Chart.UVScale.X *= UniformScale;
+				Chart.UVScale.Y *= UniformScale;
+			}
+			
+			ScaledUVArea += Chart.UVArea * Chart.UVScale.X * Chart.UVScale.Y;
+		}
+
+		if( NumMaxedOut == 0 )
+		{
+			// No charts maxed out so no need to rebalance
+			break;
+		}
+
+		if( NumMaxedOut == Charts.Num() )
+		{
+			// All charts are maxed out
+			break;
+		}
+
+		// Scale up smaller charts to maintain expected total area
+		// Want ScaledUVArea == TotalUVArea * UVScale^2
+		float RebalanceScale = UVScale * FMath::Sqrt( TotalUVArea / ScaledUVArea );
+		if( RebalanceScale < 1.01f )
+		{
+			// Stop if further rebalancing is minor
+			break;
+		}
+		UniformScale = RebalanceScale;
+	}
+#endif
+
+#if 1
+	float NonuniformScale = 1.0f;
+	while(1)
+	{
+		uint32 NumMaxedOut = 0;
+		float ScaledUVArea = 0.0f;
+		for( int32 i = 0; i < Charts.Num(); i++ )
+		{
+			FMeshChart& Chart = Charts[i];
+
+			for( int k = 0; k < 2; k++ )
+			{
+				const float MaxChartSize = TextureResolution - 1.0f;
+				const float ChartSize = Chart.MaxUV[k] - Chart.MinUV[k];
+				const float ChartSizeScaled = ChartSize * Chart.UVScale[k] * NonuniformScale;
+
+				const float Epsilon = 0.01f;
+				if( ChartSizeScaled + Epsilon > MaxChartSize )
+				{
+					// Scale oversized charts to max size
+					Chart.UVScale[k] = MaxChartSize / ChartSize;
+					NumMaxedOut++;
+				}
+				else
+				{
+					Chart.UVScale[k] *= NonuniformScale;
+				}
+			}
+
+			ScaledUVArea += Chart.UVArea * Chart.UVScale.X * Chart.UVScale.Y;
+		}
+
+		if( NumMaxedOut == 0 )
+		{
+			// No charts maxed out so no need to rebalance
+			break;
+		}
+
+		if( NumMaxedOut == Charts.Num() * 2 )
+		{
+			// All charts are maxed out in both dimensions
+			break;
+		}
+
+		// Scale up smaller charts to maintain expected total area
+		// Want ScaledUVArea == TotalUVArea * UVScale^2
+		float RebalanceScale = UVScale * FMath::Sqrt( TotalUVArea / ScaledUVArea );
+		if( RebalanceScale < 1.01f )
+		{
+			// Stop if further rebalancing is minor
+			break;
+		}
+		NonuniformScale = RebalanceScale;
+	}
+#endif
+
+	// Sort charts from largest to smallest
+	struct FCompareCharts
+	{
+		FORCEINLINE bool operator()( const FMeshChart& A, const FMeshChart& B ) const
+		{
+			// Rect area
+			FVector2D ChartRectA = ( A.MaxUV - A.MinUV ) * A.UVScale;
+			FVector2D ChartRectB = ( B.MaxUV - B.MinUV ) * B.UVScale;
+			return ChartRectA.X * ChartRectA.Y > ChartRectB.X * ChartRectB.Y;
+		}
+	};
+	Charts.Sort( FCompareCharts() );
+}
+
+bool FLayoutUV::PackCharts()
+{
+	uint32 RasterizeCycles = 0;
+	uint32 FindCycles = 0;
+	
+	double BeginPackCharts = FPlatformTime::Seconds();
+
+	LayoutRaster.Clear();
+
+	for( int32 i = 0; i < Charts.Num(); i++ )
+	{
+		FMeshChart& Chart = Charts[i];
+
+		// Try different orientations and pick best
+		int32				BestOrientation = -1;
+		FAllocator2D::FRect	BestRect = { ~0u, ~0u, ~0u, ~0u };
+
+		for( int32 Orientation = 0; Orientation < 8; Orientation++ )
+		{
+			// TODO If any dimension is less than 1 pixel shrink dimension to zero
+
+			OrientChart( Chart, Orientation );
+			
+			FVector2D ChartSize = Chart.MaxUV - Chart.MinUV;
+			ChartSize = ChartSize.X * Chart.PackingScaleU + ChartSize.Y * Chart.PackingScaleV;
+
+			// Only need half pixel dilate for rects
+			FAllocator2D::FRect	Rect;
+			Rect.X = 0;
+			Rect.Y = 0;
+			Rect.W = FMath::CeilToInt( FMath::Abs( ChartSize.X ) + 1.0f );
+			Rect.H = FMath::CeilToInt( FMath::Abs( ChartSize.Y ) + 1.0f );
+
+			// Just in case lack of precision pushes it over
+			Rect.W = FMath::Min( TextureResolution, Rect.W );
+			Rect.H = FMath::Min( TextureResolution, Rect.H );
+
+			const bool bRectPack = false;
+
+			if( bRectPack )
+			{
+				if( LayoutRaster.Find( Rect ) )
+				{
+					// Is best?
+					if( Rect.X + Rect.Y * TextureResolution < BestRect.X + BestRect.Y * TextureResolution )
+					{
+						BestOrientation = Orientation;
+						BestRect = Rect;
+					}
+				}
+				else
+				{
+					continue;
+				}
+			}
+			else
+			{
+				uint32 BeginRasterize = FPlatformTime::Cycles();
+				RasterizeChart( Chart, Rect.W, Rect.H );
+				RasterizeCycles += FPlatformTime::Cycles() - BeginRasterize;
+
+				uint32 BeginFind = FPlatformTime::Cycles();
+				bool bFound = LayoutRaster.Find( Rect, ChartRaster );
+				FindCycles += FPlatformTime::Cycles() - BeginFind;
+
+				if( bFound )
+				{
+					// Is best?
+					if( Rect.X + Rect.Y * TextureResolution < BestRect.X + BestRect.Y * TextureResolution )
+					{
+						BestOrientation = Orientation;
+						BestRect = Rect;
+					}
+				}
+				else
+				{
+					continue;
+				}
+			}
+		}
+
+		if( BestOrientation >= 0 )
+		{
+			// Add chart to layout
+			OrientChart( Chart, BestOrientation );
+
+			uint32 BeginRasterize = FPlatformTime::Cycles();
+			RasterizeChart( Chart, BestRect.W, BestRect.H );
+			RasterizeCycles += FPlatformTime::Cycles() - BeginRasterize;
+
+			LayoutRaster.Alloc( BestRect, ChartRaster );
+
+			Chart.PackingBias.X += BestRect.X;
+			Chart.PackingBias.Y += BestRect.Y;
+		}
+		else
+		{
+			// Found no orientation that fit
+			return false;
+		}
+	}
+
+	double EndPackCharts = FPlatformTime::Seconds();
+
+	UE_LOG(LogLayoutUV, Display, TEXT("PackCharts: %s"),	*FPlatformTime::PrettyTime(EndPackCharts - BeginPackCharts));
+	UE_LOG(LogLayoutUV, Display, TEXT("  Rasterize: %u"),	RasterizeCycles);
+	UE_LOG(LogLayoutUV, Display, TEXT("  Find: %u"),		FindCycles);
+
+	return true;
+}
+
+void FLayoutUV::OrientChart( FMeshChart& Chart, int32 Orientation )
+{
+	switch( Orientation )
+	{
+	case 0:
+		// 0 degrees
+		Chart.PackingScaleU = FVector2D( Chart.UVScale.X, 0 );
+		Chart.PackingScaleV = FVector2D( 0, Chart.UVScale.Y );
+		Chart.PackingBias = -Chart.MinUV.X * Chart.PackingScaleU - Chart.MinUV.Y * Chart.PackingScaleV + 0.5f;
+		break;
+	case 1:
+		// 0 degrees, flip x
+		Chart.PackingScaleU = FVector2D( -Chart.UVScale.X, 0 );
+		Chart.PackingScaleV = FVector2D( 0, Chart.UVScale.Y );
+		Chart.PackingBias = -Chart.MaxUV.X * Chart.PackingScaleU - Chart.MinUV.Y * Chart.PackingScaleV + 0.5f;
+		break;
+	case 2:
+		// 90 degrees
+		Chart.PackingScaleU = FVector2D( 0, -Chart.UVScale.X );
+		Chart.PackingScaleV = FVector2D( Chart.UVScale.Y, 0 );
+		Chart.PackingBias = -Chart.MaxUV.X * Chart.PackingScaleU - Chart.MinUV.Y * Chart.PackingScaleV + 0.5f;
+		break;
+	case 3:
+		// 90 degrees, flip x
+		Chart.PackingScaleU = FVector2D( 0, Chart.UVScale.X );
+		Chart.PackingScaleV = FVector2D( Chart.UVScale.Y, 0 );
+		Chart.PackingBias = -Chart.MinUV.X * Chart.PackingScaleU - Chart.MinUV.Y * Chart.PackingScaleV + 0.5f;
+		break;
+	case 4:
+		// 180 degrees
+		Chart.PackingScaleU = FVector2D( -Chart.UVScale.X, 0 );
+		Chart.PackingScaleV = FVector2D( 0, -Chart.UVScale.Y );
+		Chart.PackingBias = -Chart.MaxUV.X * Chart.PackingScaleU - Chart.MaxUV.Y * Chart.PackingScaleV + 0.5f;
+		break;
+	case 5:
+		// 180 degrees, flip x
+		Chart.PackingScaleU = FVector2D( Chart.UVScale.X, 0 );
+		Chart.PackingScaleV = FVector2D( 0, -Chart.UVScale.Y );
+		Chart.PackingBias = -Chart.MinUV.X * Chart.PackingScaleU - Chart.MaxUV.Y * Chart.PackingScaleV + 0.5f;
+		break;
+	case 6:
+		// 270 degrees
+		Chart.PackingScaleU = FVector2D( 0, Chart.UVScale.X );
+		Chart.PackingScaleV = FVector2D( -Chart.UVScale.Y, 0 );
+		Chart.PackingBias = -Chart.MinUV.X * Chart.PackingScaleU - Chart.MaxUV.Y * Chart.PackingScaleV + 0.5f;
+		break;
+	case 7:
+		// 270 degrees, flip x
+		Chart.PackingScaleU = FVector2D( 0, -Chart.UVScale.X );
+		Chart.PackingScaleV = FVector2D( -Chart.UVScale.Y, 0 );
+		Chart.PackingBias = -Chart.MaxUV.X * Chart.PackingScaleU - Chart.MaxUV.Y * Chart.PackingScaleV + 0.5f;
+		break;
+	}
+}
+
+// Max of 2048x2048 due to precision
+// Dilate in 28.4 fixed point. Half pixel dilation is conservative rasterization.
+// Dilation same as Minkowski sum of triangle and square.
+template< typename TShader, int32 Dilate >
+void RasterizeTriangle( TShader& Shader, const FVector2D Points[3], int32 ScissorWidth, int32 ScissorHeight )
+{
+	const FVector2D HalfPixel( 0.5f, 0.5f );
+	FVector2D p0 = Points[0] - HalfPixel;
+	FVector2D p1 = Points[1] - HalfPixel;
+	FVector2D p2 = Points[2] - HalfPixel;
+
+	// Correct winding
+	float Facing = ( p0.X - p1.X ) * ( p2.Y - p0.Y ) - ( p0.Y - p1.Y ) * ( p2.X - p0.X );
+	if( Facing < 0.0f )
+	{
+		Swap( p0, p2 );
+	}
+
+	// 28.4 fixed point
+	const int32 X0 = (int32)( 16.0f * p0.X + 0.5f );
+	const int32 X1 = (int32)( 16.0f * p1.X + 0.5f );
+	const int32 X2 = (int32)( 16.0f * p2.X + 0.5f );
+	
+	const int32 Y0 = (int32)( 16.0f * p0.Y + 0.5f );
+	const int32 Y1 = (int32)( 16.0f * p1.Y + 0.5f );
+	const int32 Y2 = (int32)( 16.0f * p2.Y + 0.5f );
+
+	// Bounding rect
+	int32 MinX = ( FMath::Min3( X0, X1, X2 ) - Dilate + 15 ) / 16;
+	int32 MaxX = ( FMath::Max3( X0, X1, X2 ) + Dilate + 15 ) / 16;
+	int32 MinY = ( FMath::Min3( Y0, Y1, Y2 ) - Dilate + 15 ) / 16;
+	int32 MaxY = ( FMath::Max3( Y0, Y1, Y2 ) + Dilate + 15 ) / 16;
+
+	// Clip to image
+	MinX = FMath::Clamp( MinX, 0, ScissorWidth );
+	MaxX = FMath::Clamp( MaxX, 0, ScissorWidth );
+	MinY = FMath::Clamp( MinY, 0, ScissorHeight );
+	MaxY = FMath::Clamp( MaxY, 0, ScissorHeight );
+
+	// Deltas
+	const int32 DX01 = X0 - X1;
+	const int32 DX12 = X1 - X2;
+	const int32 DX20 = X2 - X0;
+
+	const int32 DY01 = Y0 - Y1;
+	const int32 DY12 = Y1 - Y2;
+	const int32 DY20 = Y2 - Y0;
+
+	// Half-edge constants
+	int32 C0 = DY01 * X0 - DX01 * Y0;
+	int32 C1 = DY12 * X1 - DX12 * Y1;
+	int32 C2 = DY20 * X2 - DX20 * Y2;
+
+	// Correct for fill convention
+	C0 += ( DY01 < 0 || ( DY01 == 0 && DX01 > 0 ) ) ? 0 : -1;
+	C1 += ( DY12 < 0 || ( DY12 == 0 && DX12 > 0 ) ) ? 0 : -1;
+	C2 += ( DY20 < 0 || ( DY20 == 0 && DX20 > 0 ) ) ? 0 : -1;
+
+	// Dilate edges
+	C0 += ( abs(DX01) + abs(DY01) ) * Dilate;
+	C1 += ( abs(DX12) + abs(DY12) ) * Dilate;
+	C2 += ( abs(DX20) + abs(DY20) ) * Dilate;
+
+	for( int32 y = MinY; y < MaxY; y++ )
+	{
+		for( int32 x = MinX; x < MaxX; x++ )
+		{
+			// same as Edge1 >= 0 && Edge2 >= 0 && Edge3 >= 0
+			int32 IsInside;
+			IsInside  = C0 + (DX01 * y - DY01 * x) * 16;
+			IsInside |= C1 + (DX12 * y - DY12 * x) * 16;
+			IsInside |= C2 + (DX20 * y - DY20 * x) * 16;
+
+			if( IsInside >= 0 )
+			{
+				Shader.Process( x, y );
+			}
+		}
+	}
+}
+
+void FLayoutUV::RasterizeChart( const FMeshChart& Chart, uint32 RectW, uint32 RectH )
+{
+	// Bilinear footprint is -1 to 1 pixels. If packed geometrically, only a half pixel dilation
+	// would be needed to guarantee all charts were at least 1 pixel away, safe for bilinear filtering.
+	// Unfortunately, with pixel packing a full 1 pixel dilation is required unless chart edges exactly
+	// align with pixel centers.
+
+	ChartRaster.Clear();
+
+	for( uint32 Tri = Chart.FirstTri; Tri < Chart.LastTri; Tri++ )
+	{
+		FVector2D Points[3];
+		for( int k = 0; k < 3; k++ )
+		{
+			const FVector2D& UV = TexCoords[ 3 * SortedTris[ Tri ] + k ];
+			Points[k] = UV.X * Chart.PackingScaleU + UV.Y * Chart.PackingScaleV + Chart.PackingBias;
+		}
+
+		RasterizeTriangle< FAllocator2DShader, 16 >( ChartShader, Points, RectW, RectH );
+	}
+}
+
+void FLayoutUV::CommitPackedUVs()
+{
+	// Alloc new UV channel
+	RawMesh->WedgeTexCoords[ DstChannel ].Init( TexCoords.Num() );
+
+	// Commit chart UVs
+	for( int32 i = 0; i < Charts.Num(); i++ )
+	{
+		FMeshChart& Chart = Charts[i];
+
+		Chart.PackingScaleU /= TextureResolution;
+		Chart.PackingScaleV /= TextureResolution;
+		Chart.PackingBias /= TextureResolution;
+
+		for( uint32 Tri = Chart.FirstTri; Tri < Chart.LastTri; Tri++ )
+		{
+			for( int k = 0; k < 3; k++ )
+			{
+				uint32 Index = 3 * SortedTris[ Tri ] + k;
+				const FVector2D& UV = TexCoords[ Index ];
+				RawMesh->WedgeTexCoords[ DstChannel ][ Index ] = UV.X * Chart.PackingScaleU + UV.Y * Chart.PackingScaleV + Chart.PackingBias;
+			}
+		}
+	}
+}
\ No newline at end of file
diff --git a/Engine/Source/Developer/MeshUtilities/Private/LayoutUV.h b/Engine/Source/Developer/MeshUtilities/Private/LayoutUV.h
new file mode 100644
--- /dev/null
+++ b/Engine/Source/Developer/MeshUtilities/Private/LayoutUV.h
@@ -0,0 +1,127 @@
+// Copyright 1998-2014 Epic Games, Inc. All Rights Reserved.
+
+#pragma once
+
+#include "Allocator2D.h"
+
+struct FMeshChart
+{
+	uint32		FirstTri;
+	uint32		LastTri;
+	
+	FVector2D	MinUV;
+	FVector2D	MaxUV;
+	
+	float		UVArea;
+	FVector2D	UVScale;
+	FVector2D	WorldScale;
+	
+	FVector2D	PackingScaleU;
+	FVector2D	PackingScaleV;
+	FVector2D	PackingBias;
+
+	int32		Join[4];
+};
+
+struct FAllocator2DShader
+{
+	FAllocator2D*	Allocator2D;
+
+	FAllocator2DShader( FAllocator2D* InAllocator2D )
+	: Allocator2D( InAllocator2D )
+	{}
+
+	FORCEINLINE void Process( uint32 x, uint32 y )
+	{
+		Allocator2D->SetBit( x, y );
+	}
+};
+
+class FLayoutUV
+{
+public:
+				FLayoutUV( FRawMesh* InMesh, uint32 InSrcChannel, uint32 InDstChannel, uint32 InTextureResolution );
+
+	void		FindCharts();
+	bool		FindBestPacking();
+	void		CommitPackedUVs();
+
+private:
+	bool		PositionsMatch( uint32 a, uint32 b ) const;
+	bool		NormalsMatch( uint32 a, uint32 b ) const;
+	bool		UVsMatch( uint32 a, uint32 b ) const;
+	bool		VertsMatch( uint32 a, uint32 b ) const;
+	float		TriangleUVArea( uint32 Tri ) const;
+	void		DisconnectChart( FMeshChart& Chart, uint32 Side );
+
+	void		ScaleCharts( float UVScale );
+	bool		PackCharts();
+	void		OrientChart( FMeshChart& Chart, int32 Orientation );
+	void		RasterizeChart( const FMeshChart& Chart, uint32 RectW, uint32 RectH );
+
+	FRawMesh*	RawMesh;
+	uint32		SrcChannel;
+	uint32		DstChannel;
+	uint32		TextureResolution;
+	
+	TArray< FVector2D >		TexCoords;
+	TArray< uint32 >		SortedTris;
+	TArray< FMeshChart >	Charts;
+	float					TotalUVArea;
+	float					MaxChartSize;
+
+	FAllocator2D		LayoutRaster;
+	FAllocator2D		ChartRaster;
+	FAllocator2DShader	ChartShader;
+};
+
+
+#define THRESH_UVS_ARE_SAME (1.0f / 1024.0f)
+
+inline bool FLayoutUV::PositionsMatch( uint32 a, uint32 b ) const
+{
+	return ( RawMesh->GetWedgePosition(a) - RawMesh->GetWedgePosition(b) ).IsNearlyZero( THRESH_POINTS_ARE_SAME );
+}
+
+inline bool FLayoutUV::NormalsMatch( uint32 a, uint32 b ) const
+{
+	if( RawMesh->WedgeTangentZ.Num() != RawMesh->WedgeTexCoords[ SrcChannel ].Num() )
+	{
+		return true;
+	}
+	
+	return ( RawMesh->WedgeTangentZ[a] - RawMesh->WedgeTangentZ[b] ).IsNearlyZero( THRESH_NORMALS_ARE_SAME );
+}
+
+inline bool FLayoutUV::UVsMatch( uint32 a, uint32 b ) const
+{
+	return ( RawMesh->WedgeTexCoords[ SrcChannel ][a] - RawMesh->WedgeTexCoords[ SrcChannel ][b] ).IsNearlyZero( THRESH_UVS_ARE_SAME );
+}
+
+inline bool FLayoutUV::VertsMatch( uint32 a, uint32 b ) const
+{
+	return PositionsMatch( a, b ) && UVsMatch( a, b );
+}
+
+// Signed UV area
+inline float FLayoutUV::TriangleUVArea( uint32 Tri ) const
+{
+	FVector2D UVs[3];
+	for( int k = 0; k < 3; k++ )
+	{
+		UVs[k] = RawMesh->WedgeTexCoords[ SrcChannel ][ 3 * Tri + k ];
+	}
+
+	FVector2D EdgeUV1 = UVs[1] - UVs[0];
+	FVector2D EdgeUV2 = UVs[2] - UVs[0];
+	return 0.5f * ( EdgeUV1.X * EdgeUV2.Y - EdgeUV1.Y * EdgeUV2.X );
+}
+
+inline void FLayoutUV::DisconnectChart( FMeshChart& Chart, uint32 Side )
+{
+	if( Chart.Join[ Side ] != -1 )
+	{
+		Charts[ Chart.Join[ Side ] ].Join[ Side ^ 1 ] = -1;
+		Chart.Join[ Side ] = -1;
+	}
+}
\ No newline at end of file
diff --git a/Engine/Source/Developer/MeshUtilities/Private/MeshUtilities.cpp b/Engine/Source/Developer/MeshUtilities/Private/MeshUtilities.cpp
--- a/Engine/Source/Developer/MeshUtilities/Private/MeshUtilities.cpp
+++ b/Engine/Source/Developer/MeshUtilities/Private/MeshUtilities.cpp
@@ -14,6 +14,7 @@
 #include "ImageUtils.h"
 #include "MaterialExportUtils.h"
 #include "Textures/TextureAtlas.h"
+#include "LayoutUV.h"
 
 #include "FbxErrors.h"
 
@@ -34,6 +35,8 @@
 
 DEFINE_LOG_CATEGORY_STATIC(LogMeshUtilities,Verbose,All);
 
+#define LOCTEXT_NAMESPACE "MeshUtils"
+
 // CVars
 static TAutoConsoleVariable<int32> CVarTriangleOrderOptimization(
 	TEXT("r.TriangleOrderOptimization"),
@@ -1494,7 +1497,7 @@ static void ComputeTriangleTangents(
  * @param OutOverlappingCorners - Maps a corner index to the indices of all overlapping corners.
  * @param RawMesh - The mesh for which to compute overlapping corners.
  */
-static void FindOverlappingCorners(
+void FindOverlappingCorners(
 	TMultiMap<int32,int32>& OutOverlappingCorners,
 	FRawMesh const& RawMesh,
 	float ComparisonThreshold
@@ -1519,7 +1522,7 @@ static void FindOverlappingCorners(
 		// only need to search forward, since we add pairs both ways
 		for (int32 j = i + 1; j < VertIndexAndZ.Num(); j++)
 		{
-			if (FMath::Abs(VertIndexAndZ[j].Z - VertIndexAndZ[i].Z) > THRESH_POINTS_ARE_SAME * 4.01f)
+			if (FMath::Abs(VertIndexAndZ[j].Z - VertIndexAndZ[i].Z) > ComparisonThreshold)
 				break; // can't be any more dups
 
 			FVector PositionA = GetPositionForWedge(RawMesh, VertIndexAndZ[i].Index);
@@ -3805,16 +3808,50 @@ bool FMeshUtilities::GenerateUVs(
 #else
 	return false;
 #endif // #if PLATFORM_WINDOWS
-
 }
 
+
+
 bool FMeshUtilities::LayoutUVs(FRawMesh& RawMesh, uint32 TextureResolution, uint32 TexCoordIndex, FText& OutError)
 {
-#if PLATFORM_WINDOWS
+#if 0//PLATFORM_WINDOWS
 	FD3D9MeshUtilities D3DMeshUtils;
 	return D3DMeshUtils.LayoutUVs(RawMesh, TextureResolution, TexCoordIndex, OutError);
 #else
+
+	OutError = FText();
+	if( !RawMesh.IsValid() )
+	{
+		OutError = LOCTEXT("LayoutUVs_FailedInvalid", "LayoutUVs failed, mesh was invalid.");
+		return false;
+	}
+
+	/*int32 NumTexCoords = 0;
+	for (int32 i = 0; i < MAX_MESH_TEXTURE_COORDS; ++i)
+	{
+		if (RawMesh.WedgeTexCoords[i].Num() != RawMesh.WedgeIndices.Num())
+		{
+			break;
+		}
+		NumTexCoords++;
+	}
+
+	if (NumTexCoords > MAX_MESH_TEXTURE_COORDS || TexCoordIndex > (uint32)NumTexCoords)
+	{
+		OutError = LOCTEXT("LayoutUVs_FailedUVs", "LayoutUVs failed, incorrect number of texcoords.");
 	return false;
+	}*/
+
+	FLayoutUV Packer( &RawMesh, 0, TexCoordIndex, TextureResolution );
+
+	Packer.FindCharts();
+	bool bPackSuccess = Packer.FindBestPacking();
+	if( bPackSuccess )
+	{
+		Packer.CommitPackedUVs();
+	}
+
+	return bPackSuccess;
 #endif // #if PLATFORM_WINDOWS
 }
 
