Fixed merging for mirrored meshes
TTP# 345249

diff --git a/Engine/Source/Developer/MeshUtilities/Private/MeshUtilities.cpp b/Engine/Source/Developer/MeshUtilities/Private/MeshUtilities.cpp
--- a/Engine/Source/Developer/MeshUtilities/Private/MeshUtilities.cpp
+++ b/Engine/Source/Developer/MeshUtilities/Private/MeshUtilities.cpp
@@ -3264,11 +3264,37 @@ bool FMeshUtilities::ConstructRawMesh(
 
 	//Transform the raw mesh to world space
 	FTransform CtoM = InMeshComponent->ComponentToWorld;
+	const bool bIsMirrored = CtoM.GetDeterminant() < 0.f;
 	for (FVector& Vertex : OutRawMesh.VertexPositions)
 	{
 		Vertex = CtoM.TransformFVector4(Vertex);
 	}
 
+	if (bIsMirrored)
+	{
+		// Flip faces
+		for (int32 FaceIdx = 0; FaceIdx < OutRawMesh.WedgeIndices.Num()/3; FaceIdx++)
+		{
+			int32 I0 = FaceIdx * 3 + 0;
+			int32 I2 = FaceIdx * 3 + 2;
+			Swap(OutRawMesh.WedgeIndices[I0], OutRawMesh.WedgeIndices[I2]);
+			
+			// seems like vertex colors and UVs are not indexed, so swap values instead
+			if (OutRawMesh.WedgeColors.Num())
+			{
+				Swap(OutRawMesh.WedgeColors[I0], OutRawMesh.WedgeColors[I2]);
+			}
+
+			for (int32 i = 0; i < MAX_MESH_TEXTURE_COORDS; ++i)
+			{
+				if (OutRawMesh.WedgeTexCoords[i].Num())
+				{
+					Swap(OutRawMesh.WedgeTexCoords[i][I0], OutRawMesh.WedgeTexCoords[i][I2]);
+				}
+			}
+		}
+	}
+
 	int32 NumWedges = OutRawMesh.WedgeIndices.Num();
 	/* Always Recalculate normals, tangents and bitangents */
 	OutRawMesh.WedgeTangentZ.Empty(NumWedges);
@@ -3847,4 +3873,4 @@ void FMeshUtilities::ShutdownModule()
 }
 
 
-#undef LOCTEXT_NAMESPACE
\ No newline at end of file
+#undef LOCTEXT_NAMESPACE
