Fixed infinite loop.

diff --git a/Engine/Source/Developer/MeshUtilities/Private/LayoutUV.cpp b/Engine/Source/Developer/MeshUtilities/Private/LayoutUV.cpp
--- a/Engine/Source/Developer/MeshUtilities/Private/LayoutUV.cpp
+++ b/Engine/Source/Developer/MeshUtilities/Private/LayoutUV.cpp
@@ -593,7 +593,7 @@ void FLayoutUV::ScaleCharts( float UVScale )
 	// Scale charts such that they all fit and roughly total the same area as before
 #if 1
 	float UniformScale = 1.0f;
-	while(1)
+	for( int i = 0; i < 1000; i++ )
 	{
 		uint32 NumMaxedOut = 0;
 		float ScaledUVArea = 0.0f;
@@ -650,7 +650,7 @@ void FLayoutUV::ScaleCharts( float UVScale )
 
 #if 1
 	float NonuniformScale = 1.0f;
-	while(1)
+	for( int i = 0; i < 1000; i++ )
 	{
 		uint32 NumMaxedOut = 0;
 		float ScaledUVArea = 0.0f;
