x

diff --git a/Engine/Source/Developer/MeshUtilities/Private/MeshUtilities.cpp b/Engine/Source/Developer/MeshUtilities/Private/MeshUtilities.cpp
--- a/Engine/Source/Developer/MeshUtilities/Private/MeshUtilities.cpp
+++ b/Engine/Source/Developer/MeshUtilities/Private/MeshUtilities.cpp
@@ -103,6 +103,8 @@ private:
 		TArray<uint32>& OutPnAenIndices
 		) override;
 
+	virtual void RechunkSkeletalMeshModels(USkeletalMesh* SrcMesh, int32 MaxBonesPerChunk) override;
+
 	virtual void CalcBoneVertInfos(USkeletalMesh* SkeletalMesh, TArray<FBoneVertInfo>& Infos, bool bOnlyDominant) override;
 
 	/**
@@ -874,6 +876,61 @@ void FMeshUtilities::BuildSkeletalAdjacencyIndexBuffer(
 	}
 }
 
+void FMeshUtilities::RechunkSkeletalMeshModels(USkeletalMesh* SrcMesh, int32 MaxBonesPerChunk)
+{
+#if WITH_EDITORONLY_DATA
+	TIndirectArray<FStaticLODModel> DestModels;
+	TIndirectArray<FSkinnedModelData> ModelData;
+	FReferenceSkeleton RefSkeleton = SrcMesh->RefSkeleton;
+	uint32 VertexBufferBuildFlags = SrcMesh->GetVertexBufferFlags();
+	FSkeletalMeshResource* SrcMeshResource = SrcMesh->GetImportedResource();
+	FVector TriangleSortCenter;
+	bool bHaveTriangleSortCenter = SrcMesh->GetSortCenterPoint(TriangleSortCenter);
+
+	for (int32 ModelIndex = 0; ModelIndex < SrcMeshResource->LODModels.Num(); ++ModelIndex)
+	{
+		FSkinnedModelData& TmpModelData = *new(ModelData) FSkinnedModelData();
+		SkeletalMeshTools::CopySkinnedModelData(TmpModelData,SrcMeshResource->LODModels[ModelIndex]);
+	}
+
+	for (int32 ModelIndex = 0; ModelIndex < ModelData.Num(); ++ModelIndex)
+	{
+		TArray<FSkinnedMeshChunk*> Chunks;
+		TArray<int32> PointToOriginalMap;
+		TArray<ETriangleSortOption> SectionSortOptions;
+
+		const FSkinnedModelData& SrcModel = ModelData[ModelIndex];
+		FStaticLODModel& DestModel = *new(DestModels) FStaticLODModel();
+
+		SkeletalMeshTools::UnchunkSkeletalModel(Chunks,PointToOriginalMap,SrcModel);
+		SkeletalMeshTools::ChunkSkinnedVertices(Chunks,MaxBonesPerChunk);
+
+		for (int32 ChunkIndex = 0; ChunkIndex < Chunks.Num(); ++ChunkIndex)
+		{
+			int32 SectionIndex = Chunks[ChunkIndex]->OriginalSectionIndex;
+			SectionSortOptions.Add(SrcModel.Sections[SectionIndex].TriangleSorting);
+		}
+		check(SectionSortOptions.Num() == Chunks.Num());
+
+		BuildSkeletalModelFromChunks(DestModel,RefSkeleton,Chunks,PointToOriginalMap);
+		check(DestModel.Sections.Num() == DestModel.Chunks.Num());
+		check(DestModel.Sections.Num() == SectionSortOptions.Num());
+
+		DestModel.NumTexCoords = SrcModel.NumTexCoords;
+		DestModel.BuildVertexBuffers(VertexBufferBuildFlags);
+		for (int32 SectionIndex = 0; SectionIndex < DestModel.Sections.Num(); ++SectionIndex)
+		{
+			DestModel.SortTriangles(TriangleSortCenter,bHaveTriangleSortCenter,SectionIndex,SectionSortOptions[SectionIndex]);
+		}
+	}
+
+	//@todo-rco: Swap() doesn't seem to work
+	Exchange(SrcMeshResource->LODModels, DestModels);
+
+	// TODO: Also need to patch bEnableShadowCasting in the LODInfo struct.
+#endif // #if WITH_EDITORONLY_DATA
+}
+
 void FMeshUtilities::CalcBoneVertInfos(USkeletalMesh* SkeletalMesh, TArray<FBoneVertInfo>& Infos, bool bOnlyDominant)
 {
 	SkeletalMeshTools::CalcBoneVertInfos(SkeletalMesh,Infos,bOnlyDominant);
@@ -1767,7 +1824,6 @@ static void ComputeTangents(
 			}
 			while (NewConnections > 0);
 		}
-        
 
 		// Vertex normal construction.
 		for (int32 CornerIndex = 0; CornerIndex < 3; CornerIndex++)
@@ -2620,7 +2676,7 @@ bool FMeshUtilities::BuildSkeletalMesh( FStaticLODModel& LODModel, const FRefere
 		}
 	}
 
- 	check(Wedges.Num() == WedgeInfluenceIndices.Num());
+	check(Wedges.Num() == WedgeInfluenceIndices.Num());
 
 	// Calculate smooth wedge tangent vectors.
 
diff --git a/Engine/Source/Developer/MeshUtilities/Private/SkeletalMeshTools.h b/Engine/Source/Developer/MeshUtilities/Private/SkeletalMeshTools.h
--- a/Engine/Source/Developer/MeshUtilities/Private/SkeletalMeshTools.h
+++ b/Engine/Source/Developer/MeshUtilities/Private/SkeletalMeshTools.h
@@ -94,7 +94,5 @@ namespace SkeletalMeshTools
 	
 	void UnchunkSkeletalModel(TArray<FSkinnedMeshChunk*>& Chunks, TArray<int32>& PointToOriginalMap, const FSkinnedModelData& SrcModel);
 
-	void RechunkModels(TIndirectArray<FStaticLODModel>& DestModels, USkeletalMesh* SrcMesh, int32 MaxBonesPerChunk);
-
 	void CalcBoneVertInfos(USkeletalMesh* SkeletalMesh, TArray<FBoneVertInfo>& Infos, bool bOnlyDominant);
 };
diff --git a/Engine/Source/Developer/MeshUtilities/Public/MeshUtilities.h b/Engine/Source/Developer/MeshUtilities/Public/MeshUtilities.h
--- a/Engine/Source/Developer/MeshUtilities/Public/MeshUtilities.h
+++ b/Engine/Source/Developer/MeshUtilities/Public/MeshUtilities.h
@@ -184,6 +184,8 @@ public:
 		) = 0;
 
 	
+	virtual void RechunkSkeletalMeshModels(USkeletalMesh* SrcMesh, int32 MaxBonesPerChunk) = 0;
+
 	/**
 	 *	Calculate the verts associated weighted to each bone of the skeleton.
 	 *	The vertices returned are in the local space of the bone.
diff --git a/Engine/Source/Runtime/Engine/Classes/Engine/SkeletalMesh.h b/Engine/Source/Runtime/Engine/Classes/Engine/SkeletalMesh.h
--- a/Engine/Source/Runtime/Engine/Classes/Engine/SkeletalMesh.h
+++ b/Engine/Source/Runtime/Engine/Classes/Engine/SkeletalMesh.h
@@ -749,7 +749,7 @@ public:
 	/**
 	 * Computes flags for building vertex buffers.
 	 */
-	uint32 GetVertexBufferFlags() const;
+	ENGINE_API uint32 GetVertexBufferFlags() const;
 
 	// Begin UObject interface.
 #if WITH_EDITOR
diff --git a/Engine/Source/Runtime/Engine/Public/SkeletalMeshTypes.h b/Engine/Source/Runtime/Engine/Public/SkeletalMeshTypes.h
--- a/Engine/Source/Runtime/Engine/Public/SkeletalMeshTypes.h
+++ b/Engine/Source/Runtime/Engine/Public/SkeletalMeshTypes.h
@@ -1649,7 +1649,7 @@ public:
 	 * Initialize vertex buffers from skel mesh chunks.
 	 * @param BuildFlags See EVertexFlags.
 	 */
-	void BuildVertexBuffers(uint32 VertexFlags);
+	ENGINE_API void BuildVertexBuffers(uint32 VertexFlags);
 
 	/** Utility function for returning total number of faces in this LOD. */
 	ENGINE_API int32 GetTotalFaces() const;
