Fix lightmap UV packing to use the new method.

diff --git a/Engine/Source/Developer/MeshUtilities/Private/MeshUtilities.cpp b/Engine/Source/Developer/MeshUtilities/Private/MeshUtilities.cpp
--- a/Engine/Source/Developer/MeshUtilities/Private/MeshUtilities.cpp
+++ b/Engine/Source/Developer/MeshUtilities/Private/MeshUtilities.cpp
@@ -3388,106 +3388,6 @@ bool FMeshUtilities::ConstructRawMesh(
 /*------------------------------------------------------------------------------
 	Mesh merging 
 ------------------------------------------------------------------------------*/
-
-//
-// Helper class for generating square atlas for square lightmaps
-//
-class FLightmapPacker
-{
-public:
-	/**
-	 * @return lightmap rect in a generated atlas, invalid Rect otherwise
-	 */
-	FIntRect GetPackedLightmapRect(int32 Idx) const
-	{
-		if (PackedLigthmapSlots.IsValidIndex(Idx))
-		{
-			uint32 X0 = PackedLigthmapSlots[Idx]->X;
-			uint32 Y0 = PackedLigthmapSlots[Idx]->Y;
-			uint32 X1 = X0 + PackedLigthmapSlots[Idx]->Width;
-			uint32 Y1 = Y0 + PackedLigthmapSlots[Idx]->Height;
-			return FIntRect(X0, Y0, X1, Y1);
-		}
-		return FIntRect();
-	}
-
-	/**
-	 * @return Atlas resolution, 0 - in case atlas was not created
-	 */
-	uint32 GetAtlasResolution() const
-	{
-		return PackedLightmapAtlas ? PackedLightmapAtlas->GetWidth() : 0;
-	}
-
-	/**
-	 *  Attempts to pack provided square lightmaps into single atlas                                                                   
-	 */ 
-	bool Pack(const TArray<uint32>& LigthmapsList)
-	{
-		// Calculate total lightmaps area and sort lightmaps list by resolution
-		TArray<TPair<int32, uint32>> SortedLightmaps;
-		float TotalArea = 0;
-
-		for (int32 i = 0; i < LigthmapsList.Num(); ++i)
-		{
-			uint32 LightmapRes = LigthmapsList[i];
-			TotalArea+= FMath::Square(LightmapRes);
-			SortedLightmaps.Add(TPairInitializer<int32, uint32>(i, LightmapRes));
-		}
-		//
-		SortedLightmaps.Sort([](TPair<int32, uint32> L, TPair<int32, uint32> R) { return R.Value < L.Value; });
-		
-		// Try to pack, increasing atlas resolution with each step
-		uint32 PackedSize = FMath::RoundUpToPowerOfTwo(FMath::RoundToInt(FMath::Sqrt(TotalArea)));
-		for (int32 i = 0; i < 10; ++i) // 2 iterations should be enough >.<
-		{
-			PackedLightmapAtlas = new FLightmapAtlas(PackedSize);
-			PackedLigthmapSlots.SetNum(LigthmapsList.Num());
-
-			for (TPair<int32, uint32> SortedLightmap : SortedLightmaps)
-			{
-				const FAtlasedTextureSlot* Slot = PackedLightmapAtlas->AddLightmap(SortedLightmap.Value);
-				if (Slot == nullptr)
-				{
-					PackedLigthmapSlots.Empty();
-					PackedLightmapAtlas.Reset();
-					break;
-				}
-			
-				PackedLigthmapSlots[SortedLightmap.Key] = Slot;
-			}
-
-			if (PackedLigthmapSlots.Num() == LigthmapsList.Num())
-			{
-				return true;
-			}
-
-			PackedSize = FMath::RoundUpToPowerOfTwo(PackedSize + 1);
-		}
-		
-		return false;
-	}
-
-private:
-	struct FLightmapAtlas : public FSlateTextureAtlas
-	{
-		FLightmapAtlas(uint32 InWidth)
-			: FSlateTextureAtlas(InWidth, InWidth, 0, ESlateTextureAtlasPaddingStyle::NoPadding)
-		{}
-
-		const FAtlasedTextureSlot* AddLightmap(uint32 InWidth)
-		{
-			return FindSlotForTexture(InWidth, InWidth);
-		}
-
-		virtual void ConditionalUpdateTexture() override {};
-	};
-
-private:
-	TScopedPointer<FLightmapAtlas>			PackedLightmapAtlas;
-	TArray<const FAtlasedTextureSlot*>		PackedLigthmapSlots;
-};
-
 bool PropagatePaintedColorsToRawMesh(UStaticMeshComponent* StaticMeshComponent, int32 LODIndex, FRawMesh& RawMesh)
 {
 	UStaticMesh* StaticMesh = StaticMeshComponent->StaticMesh;
@@ -3641,34 +3541,6 @@ void FMeshUtilities::MergeActors(
 	// Attempt to pack lightmaps
 	float MergedLightMapScale = 1.f;
 	bool bCreateLightMapChannel = false;
-	FLightmapPacker LightMapPacker;	
-						
-	if (InSettings.bGenerateAtlasedLightMapUV)
-	{
-		// Set target channel for lightmap UV
-		MergedMesh.LightMapCoordinateIndex = InSettings.TargetLightMapUVChannel;
-		
-		// Collect lightmap sizes from all meshes
-		TArray<uint32> LightMapResList;
-		for (const FRawMeshExt& SourceMesh : SourceMeshes)
-		{
-			LightMapResList.Add(SourceMesh.LightMapRes);
-		}
-		
-		// Pack them into one atlas
-		LightMapPacker.Pack(LightMapResList);
-		MergedMesh.LightMapRes = LightMapPacker.GetAtlasResolution();
-
-		// Whether we need to scale down lightUV coordiantes
-		if (MergedMesh.LightMapRes > InSettings.MaxAltlasedLightMapResolution)
-		{
-			MergedLightMapScale = InSettings.MaxAltlasedLightMapResolution/(float)MergedMesh.LightMapRes;
-			MergedMesh.LightMapRes = InSettings.MaxAltlasedLightMapResolution;
-		}
-
-		// Create lightmap channel in a merged mesh only if we succeed to generate atlas for it
-		bCreateLightMapChannel = (MergedMesh.LightMapRes > 0);
-	}
 
 	// Use first mesh for naming and pivot
 	MergedMesh.AssetPackageName = SourceMeshes[0].AssetPackageName;
@@ -3718,26 +3590,6 @@ void FMeshUtilities::MergeActors(
 			}
 		}
 		
-		// Write atlased UVs into user specified TargetLightmapChannel 
-		if (bCreateLightMapChannel)
-		{
-			FVector2D	UVOffset = FVector2D::ZeroVector;
-			FIntRect	PackedLightMapRect = LightMapPacker.GetPackedLightmapRect(SourceMeshIdx);
-			if (PackedLightMapRect.Area() != 0)
-			{
-				UVOffset = FVector2D(PackedLightMapRect.Min) * MergedLightMapScale / MergedMesh.LightMapRes;
-			}
-			
-			const TArray<FVector2D>& SourceWedgeTexCoords = SourceRawMesh.WedgeTexCoords[SourceMeshes[SourceMeshIdx].LightMapCoordinateIndex];
-			
-			for (FVector2D LightMapUV : SourceWedgeTexCoords)
-			{
-				const float SourceMeshLightMapRes = SourceMeshes[SourceMeshIdx].LightMapRes;
-				const float UVScale = (SourceMeshLightMapRes * MergedLightMapScale) / MergedMesh.LightMapRes;
-				TargetRawMesh.WedgeTexCoords[MergedMesh.LightMapCoordinateIndex].Add(LightMapUV * UVScale + UVOffset);
-			}
-		}
-		
 		// Merge all other UV channels 
 		for (int32 ChannelIdx = 0; ChannelIdx < MAX_MESH_TEXTURE_COORDS; ++ChannelIdx)
 		{
@@ -3801,12 +3653,6 @@ void FMeshUtilities::MergeActors(
 		// make sure it has a new lighting guid
 		StaticMesh->LightingGuid = FGuid::NewGuid();
 
-		// Set it to use textured lightmaps. Note that Build Lighting will do the error-checking (texcoordindex exists for all LODs, etc).
-		if (bCreateLightMapChannel)
-		{
-			StaticMesh->LightMapResolution = MergedMesh.LightMapRes;
-			StaticMesh->LightMapCoordinateIndex = MergedMesh.LightMapCoordinateIndex;	
-		}
 
 		FStaticMeshSourceModel* SrcModel = new (StaticMesh->SourceModels) FStaticMeshSourceModel();
 		/*Don't allow the engine to recalculate normals*/
@@ -3814,6 +3660,10 @@ void FMeshUtilities::MergeActors(
 		SrcModel->BuildSettings.bRecomputeTangents = false;
 		SrcModel->BuildSettings.bRemoveDegenerates = false;
 		SrcModel->BuildSettings.bUseFullPrecisionUVs = false;
+		SrcModel->BuildSettings.bGenerateLightmapUVs = InSettings.bGenerateAtlasedLightMapUV;
+		SrcModel->BuildSettings.SrcLightmapIndex = InSettings.TargetLightMapUVChannel;
+		SrcModel->BuildSettings.DstLightmapIndex = InSettings.TargetLightMapUVChannel;
+
 		SrcModel->RawMeshBulkData->SaveRawMesh(MergedMesh.Mesh);
 
 		// Assign materials
