Fixed GC issues in MeshMergin tool with source meshes that use dynamic materials

diff --git a/Engine/Source/Developer/MeshUtilities/Private/MeshUtilities.cpp b/Engine/Source/Developer/MeshUtilities/Private/MeshUtilities.cpp
--- a/Engine/Source/Developer/MeshUtilities/Private/MeshUtilities.cpp
+++ b/Engine/Source/Developer/MeshUtilities/Private/MeshUtilities.cpp
@@ -3651,6 +3651,11 @@ void FMeshUtilities::MergeActors(
 		// Assign materials
 		for (UMaterialInterface* Material : UniqueMaterials)
 		{
+			if (Material && !Material->IsAsset())
+			{
+				Material = nullptr; // do not save non-asset materials
+			}
+
 			StaticMesh->Materials.Add(Material);
 		}
 		
