Suppressed layout UV log spam

diff --git a/Engine/Source/Developer/MeshUtilities/Private/LayoutUV.cpp b/Engine/Source/Developer/MeshUtilities/Private/LayoutUV.cpp
--- a/Engine/Source/Developer/MeshUtilities/Private/LayoutUV.cpp
+++ b/Engine/Source/Developer/MeshUtilities/Private/LayoutUV.cpp
@@ -4,7 +4,7 @@
 #include "LayoutUV.h"
 #include "DisjointSet.h"
 
-DEFINE_LOG_CATEGORY_STATIC(LogLayoutUV, Log, All);
+DEFINE_LOG_CATEGORY_STATIC(LogLayoutUV, Warning, All);
 
 #define CHART_JOINING	1
 
