Skel mesh importing optimizations
- Greatly improved skeletal mesh importing by removing n^2 operation when combining vertices
- Added a cvar parameter to disable triangle order optimization (debugging purposes only)

diff --git a/Engine/Source/Developer/MeshUtilities/Private/MeshUtilities.cpp b/Engine/Source/Developer/MeshUtilities/Private/MeshUtilities.cpp
--- a/Engine/Source/Developer/MeshUtilities/Private/MeshUtilities.cpp
+++ b/Engine/Source/Developer/MeshUtilities/Private/MeshUtilities.cpp
@@ -44,7 +44,8 @@ static TAutoConsoleVariable<int32> CVarTriangleOrderOptimization(
 	1,
 	TEXT("Controls the algorithm to use when optimizing the triangle order for the post-transform cache.\n")
 	TEXT("0: Use NVTriStrip (slower)\n")
-	TEXT("1: Use Forsyth algorithm (fastest)(default)"),
+	TEXT("1: Use Forsyth algorithm (fastest)(default)")
+	TEXT("2: No triangle order optimization. (least efficient, debugging purposes only)"),
 	ECVF_ReadOnly);
 
 class FMeshUtilities : public IMeshUtilities
@@ -68,6 +69,8 @@ private:
 	bool bUsingSimplygon;
 	/** True if NvTriStrip is being used for tri order optimization. */
 	bool bUsingNvTriStrip;
+	/** True if we disable triangle order optimization.  For debugging purposes only */
+	bool bDisableTriangleOrderOptimization;
 	// IMeshUtilities interface.
 	virtual const FString& GetVersionString() const override
 	{
@@ -686,7 +689,7 @@ void FMeshUtilities::CacheOptimizeIndexBuffer(TArray<uint16>& Indices)
 	{
 		NvTriStrip::CacheOptimizeIndexBuffer(Indices);
 	}
-	else
+	else if( !bDisableTriangleOrderOptimization )
 	{
 		Forsyth::CacheOptimizeIndexBuffer(Indices);
 	}
@@ -698,7 +701,7 @@ void FMeshUtilities::CacheOptimizeIndexBuffer(TArray<uint32>& Indices)
 	{
 		NvTriStrip::CacheOptimizeIndexBuffer(Indices);
 	}
-	else
+	else if( !bDisableTriangleOrderOptimization )
 	{
 		Forsyth::CacheOptimizeIndexBuffer(Indices);
 	}
@@ -2696,9 +2699,9 @@ bool FMeshUtilities::BuildSkeletalMesh( FStaticLODModel& LODModel, const FRefere
 	// to its overlapping vertices, and a table that maps a vertex to the its influenced faces
 	TMultiMap<int32,int32> Vert2Duplicates;
 	TMultiMap<int32,int32> Vert2Faces;
+	TArray<FSkeletalMeshVertIndexAndZ> VertIndexAndZ;
 	{
 		// Create a list of vertex Z/index pairs
-		TArray<FSkeletalMeshVertIndexAndZ> VertIndexAndZ;
 		VertIndexAndZ.Empty(Points.Num());
 		for (int32 i = 0; i < Points.Num(); i++)
 		{
@@ -2744,7 +2747,7 @@ bool FMeshUtilities::BuildSkeletalMesh( FStaticLODModel& LODModel, const FRefere
 		}
 
 		// we are done with this
-		VertIndexAndZ.Empty();
+		VertIndexAndZ.Reset();
 
 		// now create a map from vert indices to faces
 		for(int32 FaceIndex = 0;FaceIndex < Faces.Num();FaceIndex++)
@@ -2762,11 +2765,17 @@ bool FMeshUtilities::BuildSkeletalMesh( FStaticLODModel& LODModel, const FRefere
 	TArray<int32> DupVerts;
 	TArray<int32> DupFaces;
 
+	// List of raw calculated vertices that will be merged later
+	TArray<FSoftSkinBuildVertex> RawVertices;
+	RawVertices.Reserve( Points.Num() );
+
+	// Create a list of vertex Z/index pairs
+
 	for(int32 FaceIndex = 0;FaceIndex < Faces.Num();FaceIndex++)
 	{
 		// Only update the status progress bar if we are in the gamethread and every thousand faces. 
 		// Updating status is extremely slow
-		if( IsInGameThread() && FaceIndex % 5000 == 0 )
+		if( FaceIndex % 5000 == 0 && IsInGameThread() )
 		{
 			// Only update status if in the game thread.  When importing morph targets, this function can run in another thread
 			GWarn->StatusUpdate( FaceIndex, Faces.Num(), NSLOCTEXT("UnrealEd", "ProcessingSkeletalTriangles", "Processing Mesh Triangles") );
@@ -2852,26 +2861,6 @@ bool FMeshUtilities::BuildSkeletalMesh( FStaticLODModel& LODModel, const FRefere
 			}
 		}
 
-		// Find a chunk which matches this triangle.
-		FSkinnedMeshChunk* Chunk = NULL;
-		for (int32 i = 0; i < Chunks.Num(); ++i)
-		{
-			if (Chunks[i]->MaterialIndex == Face.MeshMaterialIndex)
-			{
-				Chunk = Chunks[i];
-				break;
-			}
-		}
-		if (Chunk == NULL)
-		{
-			Chunk = new FSkinnedMeshChunk();
-			Chunk->MaterialIndex = Face.MeshMaterialIndex;
-			Chunk->OriginalSectionIndex = Chunks.Num();
-			Chunks.Add(Chunk);
-		}
-
-		uint32 TriangleIndices[3];
-
 		for(int32 VertexIndex = 0;VertexIndex < 3;VertexIndex++)
 		{
 			FSoftSkinBuildVertex	Vertex;
@@ -2962,33 +2951,21 @@ bool FMeshUtilities::BuildSkeletalMesh( FStaticLODModel& LODModel, const FRefere
 
 			// Add the vertex as well as its original index in the points array
 			Vertex.PointWedgeIdx = Wedges[Face.iWedge[VertexIndex]].iVertex;
+			
+			int32 RawIndex = RawVertices.Add( Vertex );
 
-			int32	V = SkeletalMeshTools::AddSkinVertex(Chunk->Vertices,Vertex,bKeepOverlappingVertices);
-
+			// Add an efficient way to find dupes of this vertex later for fast combining of vertices
+			FSkeletalMeshVertIndexAndZ IAndZ;
+			IAndZ.Index = RawIndex;
+			IAndZ.Z = Vertex.Position.Z;
 
-			// set the index entry for the newly added vertex
-			// check(V >= 0 && V <= MAX_uint16);
-#if DISALLOW_32BIT_INDICES
-			if (V > MAX_uint16)
-			{
-				bTooManyVerts = true;
-			}
-			TriangleIndices[VertexIndex] = (uint16)V;
-#else
-			// TArray internally has int32 for capacity, so no need to test for uint32 as it's larger than int32
-			TriangleIndices[VertexIndex] = (uint32)V;
-#endif
-		}
-
-		if(TriangleIndices[0] != TriangleIndices[1] && TriangleIndices[0] != TriangleIndices[2] && TriangleIndices[1] != TriangleIndices[2])
-		{
-			for(uint32 VertexIndex = 0;VertexIndex < 3;VertexIndex++)
-			{
-				Chunk->Indices.Add(TriangleIndices[VertexIndex]);
-			}
+			VertIndexAndZ.Add( IAndZ );
 		}
 	}
 
+	// Generate chunks and their vertices and indices
+	SkeletalMeshTools::BuildSkeletalMeshChunks( Faces, RawVertices, VertIndexAndZ, bKeepOverlappingVertices, Chunks, bTooManyVerts );
+
 	// Chunk vertices to satisfy the requested limit.
 	static const auto MaxBonesVar = IConsoleManager::Get().FindTConsoleVariableDataInt(TEXT("Compat.MAX_GPUSKIN_BONES"));
 	const int32 MaxGPUSkinBones = MaxBonesVar->GetValueOnAnyThread();
@@ -3751,7 +3728,9 @@ void FMeshUtilities::StartupModule()
 		}
 	}
 
-	bUsingNvTriStrip = (CVarTriangleOrderOptimization.GetValueOnGameThread() == 0);
+	bDisableTriangleOrderOptimization = (CVarTriangleOrderOptimization.GetValueOnGameThread() == 2);
+
+	bUsingNvTriStrip = !bDisableTriangleOrderOptimization && (CVarTriangleOrderOptimization.GetValueOnGameThread() == 0);
 
 	// Construct and cache the version string for the mesh utilities module.
 	VersionString = FString::Printf(
diff --git a/Engine/Source/Developer/MeshUtilities/Private/SkeletalMeshTools.cpp b/Engine/Source/Developer/MeshUtilities/Private/SkeletalMeshTools.cpp
--- a/Engine/Source/Developer/MeshUtilities/Private/SkeletalMeshTools.cpp
+++ b/Engine/Source/Developer/MeshUtilities/Private/SkeletalMeshTools.cpp
@@ -7,57 +7,215 @@
 
 namespace SkeletalMeshTools
 {
-	int32 AddSkinVertex(TArray<FSoftSkinBuildVertex>& Vertices,FSoftSkinBuildVertex& Vertex, bool bKeepOverlappingVertices )
+	bool AreSkelMeshVerticesEqual( const FSoftSkinBuildVertex& V1, const FSoftSkinBuildVertex& V2 )
 	{
-		if (!bKeepOverlappingVertices)
+		if(!PointsEqual(V1.Position, V2.Position))
 		{
-			for(uint32 VertexIndex = 0;VertexIndex < (uint32)Vertices.Num();VertexIndex++)
+			return false;
+		}
+
+		bool bUVsEqual = true;
+		for(int32 UVIdx = 0; UVIdx < MAX_TEXCOORDS; ++UVIdx)
+		{
+			if(FMath::Abs(V1.UVs[UVIdx].X - V2.UVs[UVIdx].X) >(1.0f / 1024.0f))
 			{
-				FSoftSkinBuildVertex&	OtherVertex = Vertices[VertexIndex];
+				bUVsEqual = false;
+			};
+
+			if(FMath::Abs(V1.UVs[UVIdx].Y - V2.UVs[UVIdx].Y) > (1.0f / 1024.0f))
+			{
+				bUVsEqual = false;
+			}
+		}
+
+		if(!bUVsEqual)
+		{
+			return false;
+		}
+
+		if(!NormalsEqual(V1.TangentX, V2.TangentX))
+		{
+			return false;
+		}
+
+		if(!NormalsEqual(V1.TangentY, V2.TangentY))
+		{
+			return false;
+		}
 
-				if(!PointsEqual(OtherVertex.Position,Vertex.Position))
-					continue;
+		if(!NormalsEqual(V1.TangentZ, V2.TangentZ))
+		{
+			return false;
+		}
 
-				bool bUVsEqual = true;
-				for( int32 UVIdx = 0; UVIdx < MAX_TEXCOORDS; ++UVIdx )
+		bool	InfluencesMatch = 1;
+		for(uint32 InfluenceIndex = 0; InfluenceIndex < MAX_TOTAL_INFLUENCES; InfluenceIndex++)
+		{
+			if(V1.InfluenceBones[InfluenceIndex] != V2.InfluenceBones[InfluenceIndex] ||
+				V1.InfluenceWeights[InfluenceIndex] != V2.InfluenceWeights[InfluenceIndex])
+			{
+				InfluencesMatch = 0;
+				break;
+			}
+		}
+
+		if(!InfluencesMatch)
+		{
+			return false;
+		}
+
+		return true;
+	}
+
+	void BuildSkeletalMeshChunks( const TArray<FMeshFace>& Faces, const TArray<FSoftSkinBuildVertex>& RawVertices, TArray<FSkeletalMeshVertIndexAndZ>& RawVertIndexAndZ, bool bKeepOverlappingVertices, TArray<FSkinnedMeshChunk*>& OutChunks, bool& bOutTooManyVerts )
+	{
+		TArray<int32> DupVerts;
+
+		TMultiMap<int32, int32> RawVerts2Dupes;
+		{
+
+			// Sorting function for vertex Z/index pairs
+			struct FCompareFSkeletalMeshVertIndexAndZ
+			{
+				FORCEINLINE bool operator()(const FSkeletalMeshVertIndexAndZ& A, const FSkeletalMeshVertIndexAndZ& B) const
+				{
+					return A.Z < B.Z;
+				}
+			};
+
+			// Sort the vertices by z value
+			RawVertIndexAndZ.Sort(FCompareFSkeletalMeshVertIndexAndZ());
+
+			// Search for duplicates, quickly!
+			for(int32 i = 0; i < RawVertIndexAndZ.Num(); i++)
+			{
+				// only need to search forward, since we add pairs both ways
+				for(int32 j = i + 1; j < RawVertIndexAndZ.Num(); j++)
 				{
-					if(FMath::Abs(Vertex.UVs[UVIdx].X - OtherVertex.UVs[UVIdx].X) > (1.0f / 1024.0f))
+					if(FMath::Abs(RawVertIndexAndZ[j].Z - RawVertIndexAndZ[i].Z) > THRESH_POINTS_ARE_SAME)
 					{
-						bUVsEqual = false;
-					};
+						// our list is sorted, so there can't be any more dupes
+						break;
+					}
 
-					if(FMath::Abs(Vertex.UVs[UVIdx].Y - OtherVertex.UVs[UVIdx].Y) > (1.0f / 1024.0f))
+					// check to see if the points are really overlapping
+					if(PointsEqual(
+						RawVertices[RawVertIndexAndZ[i].Index].Position,
+						RawVertices[RawVertIndexAndZ[j].Index].Position))
 					{
-						bUVsEqual = false;
+						RawVerts2Dupes.Add(RawVertIndexAndZ[i].Index, RawVertIndexAndZ[j].Index);
+						RawVerts2Dupes.Add(RawVertIndexAndZ[j].Index, RawVertIndexAndZ[i].Index);
 					}
 				}
+			}
+		}
 
-				if( !bUVsEqual )
-					continue;
+		TMap<FSkinnedMeshChunk* , TMap<int32, int32> > ChunkToFinalVerts;
 
-				if(!NormalsEqual( OtherVertex.TangentX, Vertex.TangentX))
-					continue;
+	
+		uint32 TriangleIndices[3];
+		for(int32 FaceIndex = 0; FaceIndex < Faces.Num(); FaceIndex++)
+		{
+			const FMeshFace& Face = Faces[FaceIndex];
+
+			// Find a chunk which matches this triangle.
+			FSkinnedMeshChunk* Chunk = NULL;
+			for(int32 i = 0; i < OutChunks.Num(); ++i)
+			{
+				if(OutChunks[i]->MaterialIndex == Face.MeshMaterialIndex)
+				{
+					Chunk = OutChunks[i];
+					break;
+				}
+			}
+			if(Chunk == NULL)
+			{
+				Chunk = new FSkinnedMeshChunk();
+				Chunk->MaterialIndex = Face.MeshMaterialIndex;
+				Chunk->OriginalSectionIndex = OutChunks.Num();
+				OutChunks.Add(Chunk);
+			}
 
-				if(!NormalsEqual(OtherVertex.TangentY, Vertex.TangentY))
-					continue;
+			TMap<int32, int32>& FinalVerts = ChunkToFinalVerts.FindOrAdd( Chunk );
 
-				if(!NormalsEqual(OtherVertex.TangentZ, Vertex.TangentZ))
-					continue;
+			for(int32 VertexIndex = 0; VertexIndex < 3; ++VertexIndex)
+			{
+				int32 WedgeIndex = FaceIndex * 3 + VertexIndex;
+				const FSoftSkinBuildVertex& Vertex = RawVertices[WedgeIndex];
 
-				bool	InfluencesMatch = 1;
-				for(uint32 InfluenceIndex = 0;InfluenceIndex < MAX_TOTAL_INFLUENCES;InfluenceIndex++)
+				int32 FinalVertIndex = INDEX_NONE;
+				if(bKeepOverlappingVertices)
+				{
+					FinalVertIndex = Chunk->Vertices.Add(RawVertices[WedgeIndex]);
+				}
+				else
 				{
-					if( Vertex.InfluenceBones[InfluenceIndex] != OtherVertex.InfluenceBones[InfluenceIndex] || 
-						Vertex.InfluenceWeights[InfluenceIndex] != OtherVertex.InfluenceWeights[InfluenceIndex])
+					DupVerts.Reset();
+					RawVerts2Dupes.MultiFind(WedgeIndex, DupVerts);
+					DupVerts.Sort();
+
+
+					for(int32 k = 0; k < DupVerts.Num(); k++)
 					{
-						InfluencesMatch = 0;
-						break;
+						if(DupVerts[k] >= WedgeIndex)
+						{
+							// the verts beyond me haven't been placed yet, so these duplicates are not relevant
+							break;
+						}
+
+						int32 *Location = FinalVerts.Find(DupVerts[k]);
+						if(Location != NULL)
+						{
+							if(SkeletalMeshTools::AreSkelMeshVerticesEqual(Vertex, Chunk->Vertices[*Location]))
+							{
+								FinalVertIndex = *Location;
+								break;
+							}
+						}
+					}
+					if(FinalVertIndex == INDEX_NONE)
+					{
+						FinalVertIndex = Chunk->Vertices.Add(Vertex);
+						FinalVerts.Add(WedgeIndex, FinalVertIndex);
 					}
+
+				}
+
+				// set the index entry for the newly added vertex
+#if DISALLOW_32BIT_INDICES
+				if(FinalVertIndex > MAX_uint16)
+				{
+					bOutTooManyVerts = true;
 				}
-				if(!InfluencesMatch)
-					continue;
+				TriangleIndices[VertexIndex] = (uint16)FinalVertIndex;
+#else
+				// TArray internally has int32 for capacity, so no need to test for uint32 as it's larger than int32
+				TriangleIndices[VertexIndex] = (uint32)FinalVertIndex;
+#endif
+			}
 
-				return VertexIndex;
+			if(TriangleIndices[0] != TriangleIndices[1] && TriangleIndices[0] != TriangleIndices[2] && TriangleIndices[1] != TriangleIndices[2])
+			{
+				for(uint32 VertexIndex = 0; VertexIndex < 3; VertexIndex++)
+				{
+					Chunk->Indices.Add(TriangleIndices[VertexIndex]);
+				}
+			}
+		}
+	}
+
+	int32 AddSkinVertex(TArray<FSoftSkinBuildVertex>& Vertices,FSoftSkinBuildVertex& Vertex, bool bKeepOverlappingVertices )
+	{
+		if (!bKeepOverlappingVertices)
+		{
+			for(uint32 VertexIndex = 0;VertexIndex < (uint32)Vertices.Num();VertexIndex++)
+			{
+				FSoftSkinBuildVertex&	OtherVertex = Vertices[VertexIndex];
+
+				if( AreSkelMeshVerticesEqual( Vertex, OtherVertex ) )
+				{
+					return VertexIndex; 
+				}
 			}
 		}
 
diff --git a/Engine/Source/Developer/MeshUtilities/Private/SkeletalMeshTools.h b/Engine/Source/Developer/MeshUtilities/Private/SkeletalMeshTools.h
--- a/Engine/Source/Developer/MeshUtilities/Private/SkeletalMeshTools.h
+++ b/Engine/Source/Developer/MeshUtilities/Private/SkeletalMeshTools.h
@@ -78,6 +78,23 @@ namespace SkeletalMeshTools
 		return 1;
 	}
 
+	/** @return true if V1 and V2 are equal */
+	bool AreSkelMeshVerticesEqual( const FSoftSkinBuildVertex& V1, const FSoftSkinBuildVertex& V2 );
+
+	/**
+	 * Creates chunks and populates the vertex and index arrays inside each chunk
+	 *
+	 * @param Faces						List of raw faces
+	 * @param RawVertices				List of raw created, unordered, unwelded vertices
+	 * @param RawVertIndexAndZ			List of indices into the RawVertices array and each raw vertex Z position.  This is used for fast lookup of overlapping vertices
+	 * @param bKeepOverlappingVertices	Whether or not to do the overlapping vertices or not.  If true each vertex is considered unique
+	 * @param OutChunks					Created array of chunks
+	 */
+	void BuildSkeletalMeshChunks( const TArray<FMeshFace>& Faces, const TArray<FSoftSkinBuildVertex>& RawVertices, TArray<FSkeletalMeshVertIndexAndZ>& RawVertIndexAndZ, bool bKeepOverlappingVertices, TArray<FSkinnedMeshChunk*>& OutChunks, bool& bOutTooManyVerts );
+
+	/**
+	 * Adds a new vertex to the list of vertices.  Note this method is very slow with a lot of vertices as this will amount to n^2 checks that the passed in vertex is not the same as another vertex in the vertices list
+	 */
 	int32 AddSkinVertex(TArray<FSoftSkinBuildVertex>& Vertices,FSoftSkinBuildVertex& Vertex, bool bKeepOverlappingVertices );
 
 	/**
