Allow the "r.triangleorderoptimization" cvar to be set for debugging

diff --git a/Engine/Source/Developer/MeshUtilities/Private/MeshUtilities.cpp b/Engine/Source/Developer/MeshUtilities/Private/MeshUtilities.cpp
--- a/Engine/Source/Developer/MeshUtilities/Private/MeshUtilities.cpp
+++ b/Engine/Source/Developer/MeshUtilities/Private/MeshUtilities.cpp
@@ -46,7 +46,7 @@ static TAutoConsoleVariable<int32> CVarTriangleOrderOptimization(
 	TEXT("0: Use NVTriStrip (slower)\n")
 	TEXT("1: Use Forsyth algorithm (fastest)(default)")
 	TEXT("2: No triangle order optimization. (least efficient, debugging purposes only)"),
-	ECVF_ReadOnly);
+	ECVF_Default);
 
 class FMeshUtilities : public IMeshUtilities
 {
