Support for using MikkTSpace to compute tangent bases for static meshes. Currently off by default, must be enabled with r.MikkTSpace=1 until tested on existing art content.

diff --git a/Engine/Source/Developer/MeshUtilities/MeshUtilities.Build.cs b/Engine/Source/Developer/MeshUtilities/MeshUtilities.Build.cs
--- a/Engine/Source/Developer/MeshUtilities/MeshUtilities.Build.cs
+++ b/Engine/Source/Developer/MeshUtilities/MeshUtilities.Build.cs
@@ -29,7 +29,8 @@ public class MeshUtilities : ModuleRules
 
 		AddThirdPartyPrivateStaticDependencies(Target, "nvTriStrip");
 		AddThirdPartyPrivateStaticDependencies(Target, "ForsythTriOptimizer");
-		AddThirdPartyPrivateStaticDependencies(Target, "MeshSimplifier");      
+		AddThirdPartyPrivateStaticDependencies(Target, "MeshSimplifier");
+		AddThirdPartyPrivateStaticDependencies(Target, "MikkTSpace");
 
 		if ((Target.Platform == UnrealTargetPlatform.Win64) || (Target.Platform == UnrealTargetPlatform.Win32))
 		{
diff --git a/Engine/Source/Developer/MeshUtilities/Private/MeshUtilities.cpp b/Engine/Source/Developer/MeshUtilities/Private/MeshUtilities.cpp
--- a/Engine/Source/Developer/MeshUtilities/Private/MeshUtilities.cpp
+++ b/Engine/Source/Developer/MeshUtilities/Private/MeshUtilities.cpp
@@ -15,7 +15,7 @@
 #include "MaterialExportUtils.h"
 #include "Textures/TextureAtlas.h"
 #include "LayoutUV.h"
-
+#include "mikktspace.h"
 #include "DistanceFieldAtlas.h"
 #include "FbxErrors.h"
 
@@ -32,7 +32,7 @@
 // causes meshes to be rebuilt you MUST generate a new GUID and replace this
 // string with it.
 
-#define MESH_UTILITIES_VER TEXT("0db5412b27ab480f844cc7f0be5abaff")
+#define MESH_UTILITIES_VER TEXT("2C1BC8F50A7A43818AFE266EB43D9060")
 
 DEFINE_LOG_CATEGORY_STATIC(LogMeshUtilities,Verbose,All);
 
@@ -1902,6 +1902,468 @@ static void ComputeTangents(
 	check(RawMesh.WedgeTangentZ.Num() == NumWedges);
 }
 
+/*------------------------------------------------------------------------------
+	MikkTSpace for computing tangents.
+------------------------------------------------------------------------------*/
+
+static int MikkGetNumFaces(const SMikkTSpaceContext* Context)
+{
+	FRawMesh *UserData = (FRawMesh*)(Context->m_pUserData);
+	return UserData->WedgeIndices.Num() / 3;
+}
+
+static int MikkGetNumVertsOfFace(const SMikkTSpaceContext* Context, const int FaceIdx)
+{
+	// All of our meshes are triangles.
+	return 3;
+}
+
+static void MikkGetPosition(const SMikkTSpaceContext* Context, float Position[3], const int FaceIdx, const int VertIdx)
+{
+	FRawMesh *UserData = (FRawMesh*)(Context->m_pUserData);
+	FVector VertexPosition = UserData->GetWedgePosition(FaceIdx * 3 + VertIdx);
+	Position[0] = VertexPosition.X;
+	Position[1] = VertexPosition.Y;
+	Position[2] = VertexPosition.Z;
+}
+
+static void MikkGetNormal(const SMikkTSpaceContext* Context, float Normal[3], const int FaceIdx, const int VertIdx)
+{
+	FRawMesh *UserData = (FRawMesh*)(Context->m_pUserData);
+	FVector &VertexNormal = UserData->WedgeTangentZ[FaceIdx * 3 + VertIdx];
+	for (int32 i = 0; i < 3; ++i)
+	{
+		Normal[i] = VertexNormal[i];
+	}
+}
+
+static void MikkSetTSpaceBasic(const SMikkTSpaceContext* Context, const float Tangent[3], const float BitangentSign, const int FaceIdx, const int VertIdx)
+{
+	FRawMesh *UserData = (FRawMesh*)(Context->m_pUserData);
+	FVector &VertexTangent = UserData->WedgeTangentX[FaceIdx * 3 + VertIdx];
+	for (int32 i = 0; i < 3; ++i)
+	{
+		VertexTangent[i] = Tangent[i];
+	}
+	FVector Bitangent = BitangentSign * FVector::CrossProduct(UserData->WedgeTangentZ[FaceIdx * 3 + VertIdx], VertexTangent);
+	FVector &VertexBitangent = UserData->WedgeTangentY[FaceIdx * 3 + VertIdx];
+	for (int32 i = 0; i < 3; ++i)
+	{
+		VertexBitangent[i] = -Bitangent[i];
+	}
+}
+
+static void MikkGetTexCoord(const SMikkTSpaceContext* Context, float UV[2], const int FaceIdx, const int VertIdx)
+{
+	FRawMesh *UserData = (FRawMesh*)(Context->m_pUserData);
+	FVector2D &TexCoord = UserData->WedgeTexCoords[0][FaceIdx * 3 + VertIdx];
+	UV[0] = TexCoord.X;
+	UV[1] = TexCoord.Y;
+}
+
+// MikkTSpace implementations for skeletal meshes, where tangents/bitangents are ultimately derived from lists of attributes.
+
+// Holder for skeletal data to be passed to MikkTSpace.
+// Holds references to the wedge, face and points vectors that BuildSkeletalMesh is given.
+// Holds reference to the calculated normals array, which will be fleshed out if they've been calculated.
+// Holds reference to the newly created tangent and bitangent arrays, which MikkTSpace will fleshed out if required.
+class MikkTSpace_Skeletal_Mesh
+{
+public:
+	const TArray<FMeshWedge>	&wedges;			//Reference to wedge list.
+	const TArray<FMeshFace>		&faces;				//Reference to face list.	Also contains normal/tangent/bitanget/UV coords for each vertex of the face.
+	const TArray<FVector>		&points;			//Reference to position list.
+	bool						&bComputeNormals;	//Reference to bComputeNormals.
+	TArray<FVector>				&TangentsX;			//Reference to newly created tangents list.
+	TArray<FVector>				&TangentsY;			//Reference to newly created bitangents list.
+	TArray<FVector>				&TangentsZ;			//Reference to computed normals, will be empty otherwise.
+
+	MikkTSpace_Skeletal_Mesh(
+		const TArray<FMeshWedge>	&Wedges,
+		const TArray<FMeshFace>		&Faces,
+		const TArray<FVector>		&Points,
+		bool						&bComputeNormals,
+		TArray<FVector>				&VertexTangentsX,
+		TArray<FVector>				&VertexTangentsY,
+		TArray<FVector>				&VertexTangentsZ
+		)
+		:
+		wedges			(Wedges),
+		faces			(Faces),
+		points			(Points),
+		bComputeNormals	(bComputeNormals),
+		TangentsX		(VertexTangentsX),
+		TangentsY		(VertexTangentsY),
+		TangentsZ		(VertexTangentsZ)
+	{
+	}
+};
+
+static int MikkGetNumFaces_Skeletal(const SMikkTSpaceContext* Context)
+{
+	MikkTSpace_Skeletal_Mesh *UserData = (MikkTSpace_Skeletal_Mesh*)(Context->m_pUserData);
+	return UserData->faces.Num();
+}
+
+static int MikkGetNumVertsOfFace_Skeletal(const SMikkTSpaceContext* Context, const int FaceIdx)
+{
+	// Confirmed?
+	return 3;
+}
+
+static void MikkGetPosition_Skeletal(const SMikkTSpaceContext* Context, float Position[3], const int FaceIdx, const int VertIdx)
+{
+	MikkTSpace_Skeletal_Mesh *UserData = (MikkTSpace_Skeletal_Mesh*)(Context->m_pUserData);
+	const FVector &VertexPosition = UserData->points[UserData->wedges[UserData->faces[FaceIdx].iWedge[VertIdx]].iVertex];
+	Position[0] = VertexPosition.X;
+	Position[1] = VertexPosition.Y;
+	Position[2] = VertexPosition.Z;
+}
+
+static void MikkGetNormal_Skeletal(const SMikkTSpaceContext* Context, float Normal[3], const int FaceIdx, const int VertIdx)
+{
+	MikkTSpace_Skeletal_Mesh *UserData = (MikkTSpace_Skeletal_Mesh*)(Context->m_pUserData);
+	// Get different normals depending on whether they've been calculated or not.
+	if (UserData->bComputeNormals) {
+		FVector &VertexNormal = UserData->TangentsZ[FaceIdx*3 + VertIdx];
+		Normal[0] = VertexNormal.X;
+		Normal[1] = VertexNormal.Y;
+		Normal[2] = VertexNormal.Z;
+	}
+	else
+	{
+		const FVector &VertexNormal = UserData->faces[FaceIdx].TangentZ[VertIdx];
+		Normal[0] = VertexNormal.X;
+		Normal[1] = VertexNormal.Y;
+		Normal[2] = VertexNormal.Z;
+	}
+}
+
+static void MikkSetTSpaceBasic_Skeletal(const SMikkTSpaceContext* Context, const float Tangent[3], const float BitangentSign, const int FaceIdx, const int VertIdx)
+{
+	MikkTSpace_Skeletal_Mesh *UserData = (MikkTSpace_Skeletal_Mesh*)(Context->m_pUserData);
+	FVector &VertexTangent = UserData->TangentsX[FaceIdx*3 + VertIdx];
+	VertexTangent.X = Tangent[0];
+	VertexTangent.Y = Tangent[1];
+	VertexTangent.Z = Tangent[2];
+
+	FVector Bitangent;
+	// Get different normals depending on whether they've been calculated or not.
+	if (UserData->bComputeNormals) {
+		Bitangent = BitangentSign * FVector::CrossProduct(UserData->TangentsZ[UserData->wedges[UserData->faces[FaceIdx].iWedge[VertIdx]].iVertex], VertexTangent);
+	}
+	else
+	{
+		Bitangent = BitangentSign * FVector::CrossProduct(UserData->faces[FaceIdx].TangentZ[VertIdx], VertexTangent);
+	}
+	FVector &VertexBitangent = UserData->TangentsY[FaceIdx*3 + VertIdx];
+	// Switch the tangent space swizzle to X+Y-Z+ for legacy reasons.
+	VertexBitangent.X = -Bitangent[0];
+	VertexBitangent.Y = -Bitangent[1];
+	VertexBitangent.Z = -Bitangent[2];
+}
+
+static void MikkGetTexCoord_Skeletal(const SMikkTSpaceContext* Context, float UV[2], const int FaceIdx, const int VertIdx)
+{
+	MikkTSpace_Skeletal_Mesh *UserData = (MikkTSpace_Skeletal_Mesh*)(Context->m_pUserData);
+	const FVector2D &TexCoord = UserData->wedges[UserData->faces[FaceIdx].iWedge[VertIdx]].UVs[0];
+	UV[0] = TexCoord.X;
+	UV[1] = TexCoord.Y;
+}
+
+static void ComputeTangents_MikkTSpace(
+	FRawMesh& RawMesh,
+	TMultiMap<int32,int32> const& OverlappingCorners,
+	uint32 TangentOptions
+	)
+{
+	bool bBlendOverlappingNormals = (TangentOptions & ETangentOptions::BlendOverlappingNormals) != 0;
+	bool bIgnoreDegenerateTriangles = (TangentOptions & ETangentOptions::IgnoreDegenerateTriangles) != 0;
+	float ComparisonThreshold = bIgnoreDegenerateTriangles ? THRESH_POINTS_ARE_SAME : 0.0f;
+
+	// Compute per-triangle tangents.
+	TArray<FVector> TriangleTangentX;
+	TArray<FVector> TriangleTangentY;
+	TArray<FVector> TriangleTangentZ;
+
+	ComputeTriangleTangents(
+		TriangleTangentX,
+		TriangleTangentY,
+		TriangleTangentZ,
+		RawMesh,
+		bIgnoreDegenerateTriangles ? SMALL_NUMBER : 0.0f
+		);
+
+	// Declare these out here to avoid reallocations.
+	TArray<FFanFace> RelevantFacesForCorner[3];
+	TArray<int32> AdjacentFaces;
+	TArray<int32> DupVerts;
+
+	int32 NumWedges = RawMesh.WedgeIndices.Num();
+	int32 NumFaces = NumWedges / 3;
+
+	bool bWedgeNormals = true;
+	bool bWedgeTSpace = false;
+	for (int32 WedgeIdx = 0; WedgeIdx < RawMesh.WedgeTangentZ.Num(); ++WedgeIdx)
+	{
+		bWedgeNormals = bWedgeNormals && (!RawMesh.WedgeTangentZ[WedgeIdx].IsNearlyZero());
+	}
+
+	if (RawMesh.WedgeTangentX.Num() > 0 && RawMesh.WedgeTangentY.Num() > 0)
+	{
+		bWedgeTSpace = true;
+		for (int32 WedgeIdx = 0; WedgeIdx < RawMesh.WedgeTangentX.Num()
+			&& WedgeIdx < RawMesh.WedgeTangentY.Num(); ++WedgeIdx)
+		{
+			bWedgeTSpace = bWedgeTSpace && (!RawMesh.WedgeTangentX[WedgeIdx].IsNearlyZero()) && (!RawMesh.WedgeTangentY[WedgeIdx].IsNearlyZero());
+		}
+	}
+
+	// Allocate storage for tangents if none were provided, and calculate normals for MikkTSpace.
+	if (RawMesh.WedgeTangentZ.Num() != NumWedges || !bWedgeNormals)
+	{
+		// normals are not included, so we should calculate them
+		RawMesh.WedgeTangentZ.Empty(NumWedges);
+		RawMesh.WedgeTangentZ.AddZeroed(NumWedges);
+		// we need to calculate normals for MikkTSpace
+
+		for (int32 FaceIndex = 0; FaceIndex < NumFaces; FaceIndex++)
+		{
+			int32 WedgeOffset = FaceIndex * 3;
+			FVector CornerPositions[3];
+			FVector CornerNormal[3];
+
+			for (int32 CornerIndex = 0; CornerIndex < 3; CornerIndex++)
+			{
+				CornerNormal[CornerIndex] = FVector::ZeroVector;
+				CornerPositions[CornerIndex] = GetPositionForWedge(RawMesh, WedgeOffset + CornerIndex);
+				RelevantFacesForCorner[CornerIndex].Reset();
+			}
+
+			// Don't process degenerate triangles.
+			if (PointsEqual(CornerPositions[0], CornerPositions[1], ComparisonThreshold)
+				|| PointsEqual(CornerPositions[0], CornerPositions[2], ComparisonThreshold)
+				|| PointsEqual(CornerPositions[1], CornerPositions[2], ComparisonThreshold))
+			{
+				continue;
+			}
+
+			// No need to process triangles if tangents already exist.
+			bool bCornerHasNormal[3] = { 0 };
+			for (int32 CornerIndex = 0; CornerIndex < 3; CornerIndex++)
+			{
+				bCornerHasNormal[CornerIndex] = !RawMesh.WedgeTangentZ[WedgeOffset + CornerIndex].IsZero();
+			}
+			if (bCornerHasNormal[0] && bCornerHasNormal[1] && bCornerHasNormal[2])
+			{
+				continue;
+			}
+
+			// Start building a list of faces adjacent to this face.
+			AdjacentFaces.Reset();
+			for (int32 CornerIndex = 0; CornerIndex < 3; CornerIndex++)
+			{
+				int32 ThisCornerIndex = WedgeOffset + CornerIndex;
+				DupVerts.Reset();
+				OverlappingCorners.MultiFind(ThisCornerIndex,DupVerts);
+				DupVerts.Add(ThisCornerIndex); // I am a "dup" of myself
+				for (int32 k = 0; k < DupVerts.Num(); k++)
+				{
+					AdjacentFaces.AddUnique(DupVerts[k] / 3);
+				}
+			}
+
+			// We need to sort these here because the criteria for point equality is
+			// exact, so we must ensure the exact same order for all dups.
+			AdjacentFaces.Sort();
+
+			// Process adjacent faces
+			for (int32 AdjacentFaceIndex = 0; AdjacentFaceIndex < AdjacentFaces.Num(); AdjacentFaceIndex++)
+			{
+				int32 OtherFaceIndex = AdjacentFaces[AdjacentFaceIndex];
+				for (int32 OurCornerIndex = 0; OurCornerIndex < 3; OurCornerIndex++)
+				{
+					if (bCornerHasNormal[OurCornerIndex])
+						continue;
+
+					FFanFace NewFanFace;
+					int32 CommonIndexCount = 0;
+
+					// Check for vertices in common.
+					if (FaceIndex == OtherFaceIndex)
+					{
+						CommonIndexCount = 3;
+						NewFanFace.LinkedVertexIndex = OurCornerIndex;
+					}
+					else
+					{
+						// Check matching vertices against main vertex .
+						for (int32 OtherCornerIndex = 0; OtherCornerIndex < 3; OtherCornerIndex++)
+						{
+							if (PointsEqual(
+								CornerPositions[OurCornerIndex],
+								GetPositionForWedge(RawMesh, OtherFaceIndex * 3 + OtherCornerIndex),
+								ComparisonThreshold
+								))
+							{
+								CommonIndexCount++;
+								NewFanFace.LinkedVertexIndex = OtherCornerIndex;
+							}
+						}
+					}
+
+					// Add if connected by at least one point. Smoothing matches are considered later.
+					if (CommonIndexCount > 0)
+					{
+						NewFanFace.FaceIndex = OtherFaceIndex;
+						NewFanFace.bFilled = (OtherFaceIndex == FaceIndex); // Starter face for smoothing floodfill.
+						NewFanFace.bBlendTangents = NewFanFace.bFilled;
+						NewFanFace.bBlendNormals = NewFanFace.bFilled;
+						RelevantFacesForCorner[OurCornerIndex].Add(NewFanFace);
+					}
+				}
+			}
+
+			// Find true relevance of faces for a vertex normal by traversing
+			// smoothing-group-compatible connected triangle fans around common vertices.
+			for (int32 CornerIndex = 0; CornerIndex < 3; CornerIndex++)
+			{
+				if (bCornerHasNormal[CornerIndex])
+					continue;
+
+				int32 NewConnections;
+				do
+				{
+					NewConnections = 0;
+					for (int32 OtherFaceIdx = 0; OtherFaceIdx < RelevantFacesForCorner[CornerIndex].Num(); OtherFaceIdx++)
+					{
+						FFanFace& OtherFace = RelevantFacesForCorner[CornerIndex][OtherFaceIdx];
+						// The vertex' own face is initially the only face with bFilled == true.
+						if (OtherFace.bFilled)
+						{
+							for (int32 NextFaceIndex = 0; NextFaceIndex < RelevantFacesForCorner[CornerIndex].Num(); NextFaceIndex++)
+							{
+								FFanFace& NextFace = RelevantFacesForCorner[CornerIndex][NextFaceIndex];
+								if (!NextFace.bFilled) // && !NextFace.bBlendTangents)
+								{
+									if ((NextFaceIndex != OtherFaceIdx)
+										&& (RawMesh.FaceSmoothingMasks[NextFace.FaceIndex] & RawMesh.FaceSmoothingMasks[OtherFace.FaceIndex]))
+									{
+										int32 CommonVertices = 0;
+										int32 CommonNormalVertices = 0;
+										for (int32 OtherCornerIndex = 0; OtherCornerIndex < 3; OtherCornerIndex++)
+										{
+											for (int32 NextCornerIndex = 0; NextCornerIndex < 3; NextCornerIndex++)
+											{
+												int32 NextVertexIndex = RawMesh.WedgeIndices[NextFace.FaceIndex * 3 + NextCornerIndex];
+												int32 OtherVertexIndex = RawMesh.WedgeIndices[OtherFace.FaceIndex * 3 + OtherCornerIndex];
+												if (PointsEqual(
+													RawMesh.VertexPositions[NextVertexIndex],
+													RawMesh.VertexPositions[OtherVertexIndex],
+													ComparisonThreshold))
+												{
+													CommonVertices++;
+													if (bBlendOverlappingNormals
+														|| NextVertexIndex == OtherVertexIndex)
+													{
+														CommonNormalVertices++;
+													}
+												}
+											}
+										}
+										// Flood fill faces with more than one common vertices which must be touching edges.
+										if (CommonVertices > 1)
+										{
+											NextFace.bFilled = true;
+											NextFace.bBlendNormals = (CommonNormalVertices > 1);
+											NewConnections++;
+										}
+									}
+								}
+							}
+						}
+					}
+				}
+				while (NewConnections > 0);
+			}
+
+
+			// Vertex normal construction.
+			for (int32 CornerIndex = 0; CornerIndex < 3; CornerIndex++)
+			{
+				if (bCornerHasNormal[CornerIndex])
+				{
+					CornerNormal[CornerIndex] = RawMesh.WedgeTangentZ[WedgeOffset + CornerIndex];
+				}
+				else
+				{
+					for (int32 RelevantFaceIdx = 0; RelevantFaceIdx < RelevantFacesForCorner[CornerIndex].Num(); RelevantFaceIdx++)
+					{
+						FFanFace const& RelevantFace = RelevantFacesForCorner[CornerIndex][RelevantFaceIdx];
+						if (RelevantFace.bFilled)
+						{
+							int32 OtherFaceIndex = RelevantFace.FaceIndex;
+							if (RelevantFace.bBlendNormals)
+							{
+								CornerNormal[CornerIndex] += TriangleTangentZ[OtherFaceIndex];
+							}
+						}
+					}
+					if (!RawMesh.WedgeTangentZ[WedgeOffset + CornerIndex].IsZero())
+					{
+						CornerNormal[CornerIndex] = RawMesh.WedgeTangentZ[WedgeOffset + CornerIndex];
+					}
+				}
+			}
+
+			// Normalization.
+			for (int32 CornerIndex = 0; CornerIndex < 3; CornerIndex++)
+			{
+				CornerNormal[CornerIndex].Normalize();
+			}
+
+			// Copy back to the mesh.
+			for (int32 CornerIndex = 0; CornerIndex < 3; CornerIndex++)
+			{
+				RawMesh.WedgeTangentZ[WedgeOffset + CornerIndex] = CornerNormal[CornerIndex];
+			}
+		}
+	}
+
+	if (RawMesh.WedgeTangentX.Num() != NumWedges)
+	{
+		RawMesh.WedgeTangentX.Empty(NumWedges);
+		RawMesh.WedgeTangentX.AddZeroed(NumWedges);
+	}
+	if (RawMesh.WedgeTangentY.Num() != NumWedges)
+	{
+		RawMesh.WedgeTangentY.Empty(NumWedges);
+		RawMesh.WedgeTangentY.AddZeroed(NumWedges);
+	}
+
+	if (!bWedgeTSpace)
+	{
+		// we can use mikktspace to calculate the tangents
+		SMikkTSpaceInterface MikkTInterface;
+		MikkTInterface.m_getNormal = MikkGetNormal;
+		MikkTInterface.m_getNumFaces = MikkGetNumFaces;
+		MikkTInterface.m_getNumVerticesOfFace = MikkGetNumVertsOfFace;
+		MikkTInterface.m_getPosition = MikkGetPosition;
+		MikkTInterface.m_getTexCoord = MikkGetTexCoord;
+		MikkTInterface.m_setTSpaceBasic = MikkSetTSpaceBasic;
+		MikkTInterface.m_setTSpace = nullptr;
+		
+		SMikkTSpaceContext MikkTContext;
+		MikkTContext.m_pInterface = &MikkTInterface;
+		MikkTContext.m_pUserData = (void*)(&RawMesh);
+		genTangSpaceDefault(&MikkTContext);
+	}
+
+	check(RawMesh.WedgeTangentX.Num() == NumWedges);
+	check(RawMesh.WedgeTangentY.Num() == NumWedges);
+	check(RawMesh.WedgeTangentZ.Num() == NumWedges);
+}
+
 static void ComputeStreamingTextureFactors(
 	float* OutStreamingTextureFactors,
 	float* OutMaxStreamingTextureFactor,
@@ -2335,7 +2797,14 @@ bool FMeshUtilities::BuildStaticMesh(
 					// If removing degenerate triangles, ignore them when computing tangents.
 					TangentOptions |= ETangentOptions::IgnoreDegenerateTriangles;
 				}
-				ComputeTangents(RawMesh, OverlappingCorners, TangentOptions);
+				if (SrcModel.BuildSettings.bUseMikkTSpace)
+				{
+					ComputeTangents_MikkTSpace(RawMesh, OverlappingCorners, TangentOptions);
+				}
+				else
+				{
+					ComputeTangents(RawMesh, OverlappingCorners, TangentOptions);
+				}
 			}
 
 			// At this point the mesh will have valid tangents.
diff --git a/Engine/Source/Editor/StaticMeshEditor/Private/StaticMeshEditorTools.cpp b/Engine/Source/Editor/StaticMeshEditor/Private/StaticMeshEditorTools.cpp
--- a/Engine/Source/Editor/StaticMeshEditor/Private/StaticMeshEditorTools.cpp
+++ b/Engine/Source/Editor/StaticMeshEditor/Private/StaticMeshEditorTools.cpp
@@ -312,6 +312,8 @@ void FMeshBuildSettingsLayout::GenerateHeaderRowContent( FDetailWidgetRow& NodeR
 	];
 }
 
+TAutoConsoleVariable<int32> GEnableMikkTSpaceCVar(TEXT("r.MikkTSPace"),0,TEXT("Set to be non-zero to display the option of using MikkTSpace to generate tangents for static meshes."),ECVF_Default);
+
 void FMeshBuildSettingsLayout::GenerateChildContent( IDetailChildrenBuilder& ChildrenBuilder )
 {
 	{
@@ -347,6 +349,23 @@ void FMeshBuildSettingsLayout::GenerateChildContent( IDetailChildrenBuilder& Chi
 		];
 	}
 
+	if (GEnableMikkTSpaceCVar.GetValueOnAnyThread())
+	{
+		ChildrenBuilder.AddChildContent( LOCTEXT("UseMikkTSpace", "Use MikkTSpace").ToString() )
+		.NameContent()
+		[
+			SNew(STextBlock)
+			.Font( IDetailLayoutBuilder::GetDetailFont() )
+			.Text(LOCTEXT("UseMikkTSpace", "Use MikkTSpace"))
+		]
+		.ValueContent()
+		[
+			SNew(SCheckBox)
+			.IsChecked(this, &FMeshBuildSettingsLayout::ShouldUseMikkTSpace)
+			.OnCheckStateChanged(this, &FMeshBuildSettingsLayout::OnUseMikkTSpaceChanged)
+		];
+	}
+
 	{
 		ChildrenBuilder.AddChildContent( LOCTEXT("RemoveDegenerates", "Remove Degenerates") )
 		.NameContent()
@@ -554,6 +573,11 @@ ESlateCheckBoxState::Type FMeshBuildSettingsLayout::ShouldRecomputeTangents() co
 	return BuildSettings.bRecomputeTangents ? ESlateCheckBoxState::Checked : ESlateCheckBoxState::Unchecked;
 }
 
+ESlateCheckBoxState::Type FMeshBuildSettingsLayout::ShouldUseMikkTSpace() const
+{
+	return BuildSettings.bUseMikkTSpace ? ESlateCheckBoxState::Checked : ESlateCheckBoxState::Unchecked;
+}
+
 ESlateCheckBoxState::Type FMeshBuildSettingsLayout::ShouldRemoveDegenerates() const
 {
 	return BuildSettings.bRemoveDegenerates ? ESlateCheckBoxState::Checked : ESlateCheckBoxState::Unchecked;
@@ -635,6 +659,15 @@ void FMeshBuildSettingsLayout::OnRecomputeTangentsChanged(ESlateCheckBoxState::T
 	}
 }
 
+void FMeshBuildSettingsLayout::OnUseMikkTSpaceChanged(ESlateCheckBoxState::Type NewState)
+{
+	const bool bUseMikkTSpace = (NewState == ESlateCheckBoxState::Checked) ? true : false;
+	if (BuildSettings.bUseMikkTSpace != bUseMikkTSpace)
+	{
+		BuildSettings.bUseMikkTSpace = bUseMikkTSpace;
+	}
+}
+
 void FMeshBuildSettingsLayout::OnRemoveDegeneratesChanged(ESlateCheckBoxState::Type NewState)
 {
 	const bool bRemoveDegenerates = (NewState == ESlateCheckBoxState::Checked) ? true : false;
diff --git a/Engine/Source/Editor/StaticMeshEditor/Private/StaticMeshEditorTools.h b/Engine/Source/Editor/StaticMeshEditor/Private/StaticMeshEditorTools.h
--- a/Engine/Source/Editor/StaticMeshEditor/Private/StaticMeshEditorTools.h
+++ b/Engine/Source/Editor/StaticMeshEditor/Private/StaticMeshEditorTools.h
@@ -142,6 +142,7 @@ private:
 	FReply OnApplyChanges();
 	ESlateCheckBoxState::Type ShouldRecomputeNormals() const;
 	ESlateCheckBoxState::Type ShouldRecomputeTangents() const;
+	ESlateCheckBoxState::Type ShouldUseMikkTSpace() const;
 	ESlateCheckBoxState::Type ShouldRemoveDegenerates() const;
 	ESlateCheckBoxState::Type ShouldUseFullPrecisionUVs() const;
 	ESlateCheckBoxState::Type ShouldGenerateLightmapUVs() const;
@@ -156,6 +157,7 @@ private:
 
 	void OnRecomputeNormalsChanged(ESlateCheckBoxState::Type NewState);
 	void OnRecomputeTangentsChanged(ESlateCheckBoxState::Type NewState);
+	void OnUseMikkTSpaceChanged(ESlateCheckBoxState::Type NewState);
 	void OnRemoveDegeneratesChanged(ESlateCheckBoxState::Type NewState);
 	void OnUseFullPrecisionUVsChanged(ESlateCheckBoxState::Type NewState);
 	void OnGenerateLightmapUVsChanged(ESlateCheckBoxState::Type NewState);
diff --git a/Engine/Source/Runtime/Engine/Classes/Engine/EngineTypes.h b/Engine/Source/Runtime/Engine/Classes/Engine/EngineTypes.h
--- a/Engine/Source/Runtime/Engine/Classes/Engine/EngineTypes.h
+++ b/Engine/Source/Runtime/Engine/Classes/Engine/EngineTypes.h
@@ -1849,6 +1849,10 @@ struct FMeshBuildSettings
 {
 	GENERATED_USTRUCT_BODY()
 
+	/** If true, degenerate triangles will be removed. */
+	UPROPERTY(EditAnywhere, Category=BuildSettings)
+	bool bUseMikkTSpace;
+
 	/** If true, normals in the raw mesh are ignored and recomputed. */
 	UPROPERTY(EditAnywhere, Category=BuildSettings)
 	bool bRecomputeNormals;
@@ -1900,7 +1904,8 @@ struct FMeshBuildSettings
 
 	/** Default settings. */
 	FMeshBuildSettings()
-		: bRecomputeNormals(true)
+		: bUseMikkTSpace(false)
+		, bRecomputeNormals(true)
 		, bRecomputeTangents(true)
 		, bRemoveDegenerates(true)
 		, bUseFullPrecisionUVs(false)
@@ -1919,6 +1924,7 @@ struct FMeshBuildSettings
 	{
 		return bRecomputeNormals == Other.bRecomputeNormals
 			&& bRecomputeTangents == Other.bRecomputeTangents
+			&& bUseMikkTSpace == Other.bUseMikkTSpace
 			&& bRemoveDegenerates == Other.bRemoveDegenerates
 			&& bUseFullPrecisionUVs == Other.bUseFullPrecisionUVs
 			&& bGenerateLightmapUVs == Other.bGenerateLightmapUVs
diff --git a/Engine/Source/Runtime/Engine/Private/StaticMesh.cpp b/Engine/Source/Runtime/Engine/Private/StaticMesh.cpp
--- a/Engine/Source/Runtime/Engine/Private/StaticMesh.cpp
+++ b/Engine/Source/Runtime/Engine/Private/StaticMesh.cpp
@@ -997,6 +997,12 @@ FArchive& operator<<(FArchive& Ar, FMeshBuildSettings& BuildSettings)
 	// Note: this serializer is currently only used to build the mesh DDC key, no versioning is required
 	Ar << BuildSettings.bRecomputeNormals;
 	Ar << BuildSettings.bRecomputeTangents;
+#if PLATFORM_WINDOWS
+	Ar << BuildSettings.bUseMikkTSpace;
+#else
+	// Until the library gets compiled for Mac + Linux ignore the user's selection.
+	bool bUseMikkTSpace = false;
+#endif
 	Ar << BuildSettings.bRemoveDegenerates;
 	Ar << BuildSettings.bUseFullPrecisionUVs;
 	Ar << BuildSettings.bGenerateLightmapUVs;
