Support spline meshes in Mesh proxy generation
UE-6299

diff --git a/Engine/Source/Developer/MeshUtilities/Private/MeshUtilities.cpp b/Engine/Source/Developer/MeshUtilities/Private/MeshUtilities.cpp
--- a/Engine/Source/Developer/MeshUtilities/Private/MeshUtilities.cpp
+++ b/Engine/Source/Developer/MeshUtilities/Private/MeshUtilities.cpp
@@ -18,6 +18,7 @@
 #include "mikktspace.h"
 #include "DistanceFieldAtlas.h"
 #include "FbxErrors.h"
+#include "Components/SplineMeshComponent.h"
 
 //@todo - implement required vector intrinsics for other implementations
 #if PLATFORM_ENABLE_VECTORINTRINSICS
@@ -3781,15 +3782,28 @@ bool FMeshUtilities::ConstructRawMesh(
 		UE_LOG(LogMeshUtilities, Error, TEXT("Raw mesh (%s) is corrupt for LOD%d."), *SrcMesh->GetName(), 1);
 		return false;
 	}
+	
+	// Handle spline mesh deformation
+	if (InMeshComponent->IsA<USplineMeshComponent>())
+	{
+		USplineMeshComponent* SplineMeshComponent = Cast<USplineMeshComponent>(InMeshComponent);
+		for (int32 iVert = 0; iVert < OutRawMesh.VertexPositions.Num(); ++iVert)
+		{
+			float& Z = USplineMeshComponent::GetAxisValue(OutRawMesh.VertexPositions[iVert], SplineMeshComponent->ForwardAxis);
+			FTransform SliceTransform = SplineMeshComponent->CalcSliceTransform(Z);
+			Z = 0.0f;
+			OutRawMesh.VertexPositions[iVert] = SliceTransform.TransformPosition(OutRawMesh.VertexPositions[iVert]);
+		}
+	}
 
-	//Transform the raw mesh to world space
+	// Transform raw mesh to world space
 	FTransform CtoM = InMeshComponent->ComponentToWorld;
-	const bool bIsMirrored = CtoM.GetDeterminant() < 0.f;
 	for (FVector& Vertex : OutRawMesh.VertexPositions)
 	{
 		Vertex = CtoM.TransformFVector4(Vertex);
 	}
 
+	const bool bIsMirrored = CtoM.GetDeterminant() < 0.f;
 	if (bIsMirrored)
 	{
 		// Flip faces
@@ -3814,7 +3828,7 @@ bool FMeshUtilities::ConstructRawMesh(
 			}
 		}
 	}
-
+		
 	int32 NumWedges = OutRawMesh.WedgeIndices.Num();
 	/* Always Recalculate normals, tangents and bitangents */
 	OutRawMesh.WedgeTangentZ.Empty(NumWedges);
