Fixed importing Mixxamo Fuse characters with glitching

- Influence bones were imported when 0 weighted

diff --git a/Engine/Source/Developer/MeshUtilities/Private/MeshUtilities.cpp b/Engine/Source/Developer/MeshUtilities/Private/MeshUtilities.cpp
--- a/Engine/Source/Developer/MeshUtilities/Private/MeshUtilities.cpp
+++ b/Engine/Source/Developer/MeshUtilities/Private/MeshUtilities.cpp
@@ -1094,6 +1094,8 @@ void FMeshUtilities::BuildSkeletalModelFromChunks(FStaticLODModel& LODModel,cons
 				FMemory::Memcpy( RigidVertex.UVs, SoftVertex.UVs, sizeof(FVector2D)*MAX_TEXCOORDS );
 				RigidVertex.Color = SoftVertex.Color;
 				RigidVertex.Bone = SoftVertex.InfluenceBones[0];
+				// make sure it exists in bone map
+				check (Chunk.BoneMap.IsValidIndex(SoftVertex.InfluenceBones[0]));
 				Chunk.RigidVertices.Add(RigidVertex);
 				ChunkVertexIndexRemap[VertexIndex] = (uint32)(Chunk.BaseVertexIndex + CurrentVertexIndex);
 				CurrentVertexIndex++;
@@ -1119,8 +1121,14 @@ void FMeshUtilities::BuildSkeletalModelFromChunks(FStaticLODModel& LODModel,cons
 				NewVertex.Color = SoftVertex.Color;
 				for (int32 i = 0; i < MAX_TOTAL_INFLUENCES; ++i)
 				{
-					NewVertex.InfluenceBones[i] = SoftVertex.InfluenceBones[i];
-					NewVertex.InfluenceWeights[i] = SoftVertex.InfluenceWeights[i];
+					// it only adds to the bone map if it has weight on it
+					// BoneMap contains only the bones that has influence with weight of >0.f
+					// so here, just make sure it is included before setting the data
+					if (Chunk.BoneMap.IsValidIndex(SoftVertex.InfluenceBones[i]))
+					{
+						NewVertex.InfluenceBones[i] = SoftVertex.InfluenceBones[i];
+						NewVertex.InfluenceWeights[i] = SoftVertex.InfluenceWeights[i];
+					}
 				}
 				Chunk.SoftVertices.Add(NewVertex);
 				ChunkVertexIndexRemap[VertexIndex] = (uint32)(Chunk.BaseVertexIndex + CurrentVertexIndex);
