Fixed assert.

diff --git a/Engine/Source/Developer/MeshUtilities/Private/LayoutUV.cpp b/Engine/Source/Developer/MeshUtilities/Private/LayoutUV.cpp
--- a/Engine/Source/Developer/MeshUtilities/Private/LayoutUV.cpp
+++ b/Engine/Source/Developer/MeshUtilities/Private/LayoutUV.cpp
@@ -279,6 +279,13 @@ void FLayoutUV::FindCharts( const TMultiMap<int32,int32>& OverlappingCorners )
 						if( bMirrored )
 							continue;
 
+						if( ChartA.Join[ Side ^ 0 ] != -1 ||
+							ChartB.Join[ Side ^ 1 ] != -1 )
+						{
+							// Already joined with something else
+							continue;
+						}
+
 						uint32 Sign = Side & 1;
 						uint32 Axis = Side >> 1;
 
