Fixed crash during simplygon mesh proxy generation when one of the source meshes has no material

diff --git a/Engine/Source/Developer/MeshUtilities/Private/MeshUtilities.cpp b/Engine/Source/Developer/MeshUtilities/Private/MeshUtilities.cpp
--- a/Engine/Source/Developer/MeshUtilities/Private/MeshUtilities.cpp
+++ b/Engine/Source/Developer/MeshUtilities/Private/MeshUtilities.cpp
@@ -3838,14 +3838,13 @@ bool FMeshUtilities::ConstructRawMesh(
 	{
 		// Add material and store the material ID
 		UMaterialInterface* MaterialToAdd = InMeshComponent->GetMaterial(Section.MaterialIndex);
-		if (MaterialToAdd != NULL)
+		if (MaterialToAdd == nullptr)
 		{
-			OutGlobalMaterialIndices.Add(OutUniqueMaterials.AddUnique(MaterialToAdd));
-		}
-		else
-		{
-			OutGlobalMaterialIndices.Add(INDEX_NONE);
+			MaterialToAdd = UMaterial::GetDefaultMaterial(MD_Surface);
 		}
+		
+		int32 MaterialIdx = OutUniqueMaterials.AddUnique(MaterialToAdd);
+		OutGlobalMaterialIndices.Add(MaterialIdx);
 	}
 			
 	return true;
