Large improvement to high poly mesh importing speed.

Disabled index buffer optimization when mesh has more than 50000 triangles. These won't be game meshes and the optimization time takes too long.

Also disabled adjacency buffer generation for high poly meshes. This would be better done as an import option.

diff --git a/Engine/Source/Developer/MeshUtilities/Private/MeshUtilities.cpp b/Engine/Source/Developer/MeshUtilities/Private/MeshUtilities.cpp
--- a/Engine/Source/Developer/MeshUtilities/Private/MeshUtilities.cpp
+++ b/Engine/Source/Developer/MeshUtilities/Private/MeshUtilities.cpp
@@ -842,7 +842,7 @@ static void BuildStaticAdjacencyIndexBuffer(
 	TArray<uint32>& OutPnAenIndices
 	)
 {
-	if ( Indices.Num() )
+	if ( Indices.Num() && Indices.Num() < 50000 * 3 )
 	{
 		FStaticMeshNvRenderBuffer StaticMeshRenderBuffer( PositionVertexBuffer, VertexBuffer, Indices );
 		nv::IndexBuffer* PnAENIndexBuffer = nv::tess::buildTessellationBuffer( &StaticMeshRenderBuffer, nv::DBM_PnAenDominantCorner, true );
@@ -2946,8 +2946,12 @@ bool FMeshUtilities::BuildStaticMesh(
 			float ComparisonThreshold = GetComparisonThreshold(LODBuildSettings[LODIndex]);
 			BuildStaticMeshVertexAndIndexBuffers(Vertices, PerSectionIndices, WedgeMap, RawMesh, LODOverlappingCorners[LODIndex], ComparisonThreshold, LODBuildSettings[LODIndex].BuildScale3D );
 			check(WedgeMap.Num() == RawMesh.WedgeIndices.Num());
-			CacheOptimizeVertexAndIndexBuffer(Vertices, PerSectionIndices, WedgeMap);
-			check(WedgeMap.Num() == RawMesh.WedgeIndices.Num());
+
+			if( RawMesh.WedgeIndices.Num() < 50000 * 3 )
+			{
+				CacheOptimizeVertexAndIndexBuffer(Vertices, PerSectionIndices, WedgeMap);
+				check(WedgeMap.Num() == RawMesh.WedgeIndices.Num());
+			}
 		}
 
 		// Initialize the vertex buffer.
@@ -3013,7 +3017,12 @@ bool FMeshUtilities::BuildStaticMesh(
 				CombinedIndices,
 				LODModel.Sections
 				);
-			CacheOptimizeIndexBuffer(DepthOnlyIndices);
+
+			if( DepthOnlyIndices.Num() < 50000 * 3 )
+			{
+				CacheOptimizeIndexBuffer(DepthOnlyIndices);
+			}
+
 			LODModel.DepthOnlyIndexBuffer.SetIndices(DepthOnlyIndices, bNeeds32BitIndices ? EIndexBufferStride::Force32Bit : EIndexBufferStride::Force16Bit);
 		}
 
