UE4 - Save generated LODs for static meshes

diff --git a/Engine/Source/Developer/AssetTools/Private/AssetTypeActions/AssetTypeActions_StaticMesh.cpp b/Engine/Source/Developer/AssetTools/Private/AssetTypeActions/AssetTypeActions_StaticMesh.cpp
--- a/Engine/Source/Developer/AssetTools/Private/AssetTypeActions/AssetTypeActions_StaticMesh.cpp
+++ b/Engine/Source/Developer/AssetTools/Private/AssetTypeActions/AssetTypeActions_StaticMesh.cpp
@@ -17,10 +17,31 @@
 
 #define LOCTEXT_NAMESPACE "AssetTypeActions"
 
+static TAutoConsoleVariable<int32> CVarEnableSaveGeneratedLODsInPackage(
+	TEXT("r.StaticMesh.EnableSaveGeneratedLODsInPackage"),
+	0,
+	TEXT("Enables saving generated LODs in the Package.\n") \
+	TEXT("0 - Do not save (and hide this menu option) [default].\n") \
+	TEXT("1 - Enable this option and save the LODs in the Package.\n"),
+	ECVF_Default);
+
 void FAssetTypeActions_StaticMesh::GetActions( const TArray<UObject*>& InObjects, FMenuBuilder& MenuBuilder )
 {
 	auto Meshes = GetTypedWeakObjectPtrs<UStaticMesh>(InObjects);
 
+	if (CVarEnableSaveGeneratedLODsInPackage.GetValueOnGameThread() != 0)
+	{
+		MenuBuilder.AddMenuEntry(
+			NSLOCTEXT("AssetTypeActions_StaticMesh", "ObjectContext_SaveGeneratedLODsInPackage", "Save Generated LODs"),
+			NSLOCTEXT("AssetTypeActions_StaticMesh", "ObjectContext_SaveGeneratedLODsInPackageTooltip", "Run the mesh reduce and save the generated LODs as part of the package."),
+			FSlateIcon(),
+			FUIAction(
+				FExecuteAction::CreateSP(this, &FAssetTypeActions_StaticMesh::ExecuteSaveGeneratedLODsInPackage, Meshes),
+				FCanExecuteAction()
+				)
+			);
+	}
+
 	MenuBuilder.AddMenuEntry(
 		NSLOCTEXT("AssetTypeActions_StaticMesh", "ObjectContext_CreateDestructibleMesh", "Create Destructible Mesh"),
 		NSLOCTEXT("AssetTypeActions_StaticMesh", "ObjectContext_CreateDestructibleMeshTooltip", "Creates a DestructibleMesh from the StaticMesh and opens it in the DestructibleMesh editor."),
@@ -129,4 +150,16 @@ void FAssetTypeActions_StaticMesh::ExecuteCreateDestructibleMesh(TArray<TWeakObj
 	}
 }
 
+void FAssetTypeActions_StaticMesh::ExecuteSaveGeneratedLODsInPackage(TArray<TWeakObjectPtr<UStaticMesh>> Objects)
+{
+	for (auto StaticMeshIt = Objects.CreateConstIterator(); StaticMeshIt; ++StaticMeshIt)
+	{
+		auto StaticMesh = (*StaticMeshIt).Get();
+		if (StaticMesh)
+		{
+			StaticMesh->GenerateLodsInPackage();
+		}
+	}
+}
+
 #undef LOCTEXT_NAMESPACE
diff --git a/Engine/Source/Developer/AssetTools/Private/AssetTypeActions/AssetTypeActions_StaticMesh.h b/Engine/Source/Developer/AssetTools/Private/AssetTypeActions/AssetTypeActions_StaticMesh.h
--- a/Engine/Source/Developer/AssetTools/Private/AssetTypeActions/AssetTypeActions_StaticMesh.h
+++ b/Engine/Source/Developer/AssetTools/Private/AssetTypeActions/AssetTypeActions_StaticMesh.h
@@ -24,6 +24,9 @@ private:
 	/** Handler for when CreateDestructibleMesh is selected */
 	void ExecuteCreateDestructibleMesh(TArray<TWeakObjectPtr<UStaticMesh>> Objects);
 
+	// Handler for when SaveGeneratedLODsInPackage is selected
+	void ExecuteSaveGeneratedLODsInPackage(TArray<TWeakObjectPtr<UStaticMesh>> Objects);
+
 	/** Handler to provide the list of LODs that can be imported or reimported */
 	void GetImportLODMenu(class FMenuBuilder& MenuBuilder,TArray<TWeakObjectPtr<UStaticMesh>> Objects);
 };
\ No newline at end of file
diff --git a/Engine/Source/Developer/MeshUtilities/Private/MeshUtilities.cpp b/Engine/Source/Developer/MeshUtilities/Private/MeshUtilities.cpp
--- a/Engine/Source/Developer/MeshUtilities/Private/MeshUtilities.cpp
+++ b/Engine/Source/Developer/MeshUtilities/Private/MeshUtilities.cpp
@@ -77,12 +77,15 @@ private:
 	{
 		return VersionString;
 	}
+
 	virtual bool BuildStaticMesh(
 		FStaticMeshRenderData& OutRenderData,
 		TArray<FStaticMeshSourceModel>& SourceModels,
 		const FStaticMeshLODGroup& LODGroup
 		) override;
 
+	virtual bool GenerateStaticMeshLODs(TArray<FStaticMeshSourceModel>& Models, const FStaticMeshLODGroup& LODGroup) override;
+
 	virtual void GenerateSignedDistanceFieldVolumeData(
 		const FStaticMeshLODResources& LODModel,
 		class FQueuedThreadPool& ThreadPool,
@@ -148,6 +151,9 @@ private:
 		TArray<UMaterialInterface*>& OutUniqueMaterials,
 		TArray<int32>& OutGlobalMaterialIndices
 		) const;
+
+	// Need to call some members from this class, (which is internal to this module)
+	friend class FStaticMeshUtilityBuilder;
 };
 
 IMPLEMENT_MODULE(FMeshUtilities, MeshUtilities);
@@ -2737,353 +2743,435 @@ void FMeshUtilities::CacheOptimizeVertexAndIndexBuffer(
 	}
 }
 
-static void ApplyScaling( FRawMesh& Mesh, float BuildScale )
+class FStaticMeshUtilityBuilder
 {
-	const int32 NumFaces = Mesh.WedgeIndices.Num() / 3;
-
-	for(int32 FaceIndex = 0; FaceIndex < NumFaces; FaceIndex++)
-	{
-		for (int32 TriVertexIndex = 0; TriVertexIndex < 3; TriVertexIndex++)
-		{
-			const int32 WedgeIndex = FaceIndex * 3 + TriVertexIndex;
-			int32 VertexIndex = Mesh.WedgeIndices[WedgeIndex];
-			Mesh.VertexPositions[VertexIndex] *= BuildScale;
-		}
-	}
-
-}
-
-bool FMeshUtilities::BuildStaticMesh(
-	FStaticMeshRenderData& OutRenderData,
-	TArray<FStaticMeshSourceModel>& SourceModels,
-	const FStaticMeshLODGroup& LODGroup
-	)
-{
-	TIndirectArray<FRawMesh> LODMeshes;
-	TIndirectArray<TMultiMap<int32,int32> > LODOverlappingCorners;
-	float LODMaxDeviation[MAX_STATIC_MESH_LODS];
-	FMeshBuildSettings LODBuildSettings[MAX_STATIC_MESH_LODS];
+public:
+	FStaticMeshUtilityBuilder() : Stage(EStage::Uninit), NumValidLODs(0) {}
 
-	// Gather source meshes for each LOD.
-	for (int32 LODIndex = 0; LODIndex < SourceModels.Num(); ++LODIndex)
+	bool GatherSourceMeshesPerLOD(TArray<FStaticMeshSourceModel>& SourceModels, IMeshReduction* MeshReduction)
 	{
-		FStaticMeshSourceModel& SrcModel = SourceModels[LODIndex];
-		FRawMesh& RawMesh = *new(LODMeshes) FRawMesh;
-		TMultiMap<int32,int32>& OverlappingCorners = *new(LODOverlappingCorners) TMultiMap<int32,int32>;
+		check(Stage == EStage::Uninit);
 
-		if (!SrcModel.RawMeshBulkData->IsEmpty())
+		// Gather source meshes for each LOD.
+		for (int32 LODIndex = 0; LODIndex < SourceModels.Num(); ++LODIndex)
 		{
-			SrcModel.RawMeshBulkData->LoadRawMesh(RawMesh);
-			// Make sure the raw mesh is not irreparably malformed.
-			if (!RawMesh.IsValidOrFixable())
-			{
-				UE_LOG(LogMeshUtilities,Error,TEXT("Raw mesh is corrupt for LOD%d."), LODIndex);
-				return false;
-			}
-			LODBuildSettings[LODIndex] = SrcModel.BuildSettings;
+			FStaticMeshSourceModel& SrcModel = SourceModels[LODIndex];
+			FRawMesh& RawMesh = *new(LODMeshes)FRawMesh;
+			TMultiMap<int32, int32>& OverlappingCorners = *new(LODOverlappingCorners)TMultiMap<int32, int32>;
 
-			float ComparisonThreshold = GetComparisonThreshold(LODBuildSettings[LODIndex]);
-			int32 NumWedges = RawMesh.WedgeIndices.Num();
+			if (!SrcModel.RawMeshBulkData->IsEmpty())
+			{
+				SrcModel.RawMeshBulkData->LoadRawMesh(RawMesh);
+				// Make sure the raw mesh is not irreparably malformed.
+				if (!RawMesh.IsValidOrFixable())
+				{
+					UE_LOG(LogMeshUtilities, Error, TEXT("Raw mesh is corrupt for LOD%d."), LODIndex);
+					return false;
+				}
+				LODBuildSettings[LODIndex] = SrcModel.BuildSettings;
 
-			// Find overlapping corners to accelerate adjacency.
-			FindOverlappingCorners(OverlappingCorners, RawMesh, ComparisonThreshold);
+				float ComparisonThreshold = GetComparisonThreshold(LODBuildSettings[LODIndex]);
+				int32 NumWedges = RawMesh.WedgeIndices.Num();
 
-			// Figure out if we should recompute normals and tangents.
-			bool bRecomputeNormals = SrcModel.BuildSettings.bRecomputeNormals || RawMesh.WedgeTangentZ.Num() == 0;
-			bool bRecomputeTangents = SrcModel.BuildSettings.bRecomputeTangents || RawMesh.WedgeTangentX.Num() == 0 || RawMesh.WedgeTangentY.Num() == 0;
+				// Find overlapping corners to accelerate adjacency.
+				FindOverlappingCorners(OverlappingCorners, RawMesh, ComparisonThreshold);
 
-			// Dump normals and tangents if we are recomputing them.
-			if (bRecomputeTangents)
-			{
-				RawMesh.WedgeTangentX.Empty(NumWedges);
-				RawMesh.WedgeTangentX.AddZeroed(NumWedges);
-				RawMesh.WedgeTangentY.Empty(NumWedges);
-				RawMesh.WedgeTangentY.AddZeroed(NumWedges);
-			}
-			if (bRecomputeNormals)
-			{
-				RawMesh.WedgeTangentZ.Empty(NumWedges);
-				RawMesh.WedgeTangentZ.AddZeroed(NumWedges);
-			}
+				// Figure out if we should recompute normals and tangents.
+				bool bRecomputeNormals = SrcModel.BuildSettings.bRecomputeNormals || RawMesh.WedgeTangentZ.Num() == 0;
+				bool bRecomputeTangents = SrcModel.BuildSettings.bRecomputeTangents || RawMesh.WedgeTangentX.Num() == 0 || RawMesh.WedgeTangentY.Num() == 0;
 
-			// Compute any missing tangents.
-			{
-				// Static meshes always blend normals of overlapping corners.
-				uint32 TangentOptions = ETangentOptions::BlendOverlappingNormals;
-				if (SrcModel.BuildSettings.bRemoveDegenerates)
+				// Dump normals and tangents if we are recomputing them.
+				if (bRecomputeTangents)
 				{
-					// If removing degenerate triangles, ignore them when computing tangents.
-					TangentOptions |= ETangentOptions::IgnoreDegenerateTriangles;
+					RawMesh.WedgeTangentX.Empty(NumWedges);
+					RawMesh.WedgeTangentX.AddZeroed(NumWedges);
+					RawMesh.WedgeTangentY.Empty(NumWedges);
+					RawMesh.WedgeTangentY.AddZeroed(NumWedges);
 				}
-				if (SrcModel.BuildSettings.bUseMikkTSpace)
+				if (bRecomputeNormals)
 				{
-					ComputeTangents_MikkTSpace(RawMesh, OverlappingCorners, TangentOptions);
+					RawMesh.WedgeTangentZ.Empty(NumWedges);
+					RawMesh.WedgeTangentZ.AddZeroed(NumWedges);
 				}
-				else
+
+				// Compute any missing tangents.
 				{
-					ComputeTangents(RawMesh, OverlappingCorners, TangentOptions);
+					// Static meshes always blend normals of overlapping corners.
+					uint32 TangentOptions = ETangentOptions::BlendOverlappingNormals;
+					if (SrcModel.BuildSettings.bRemoveDegenerates)
+					{
+						// If removing degenerate triangles, ignore them when computing tangents.
+						TangentOptions |= ETangentOptions::IgnoreDegenerateTriangles;
+					}
+					if (SrcModel.BuildSettings.bUseMikkTSpace)
+					{
+						ComputeTangents_MikkTSpace(RawMesh, OverlappingCorners, TangentOptions);
+					}
+					else
+					{
+						ComputeTangents(RawMesh, OverlappingCorners, TangentOptions);
+					}
 				}
-			}
 
-			// At this point the mesh will have valid tangents.
-			check(RawMesh.WedgeTangentX.Num() == NumWedges);
-			check(RawMesh.WedgeTangentY.Num() == NumWedges);
-			check(RawMesh.WedgeTangentZ.Num() == NumWedges);
+				// At this point the mesh will have valid tangents.
+				check(RawMesh.WedgeTangentX.Num() == NumWedges);
+				check(RawMesh.WedgeTangentY.Num() == NumWedges);
+				check(RawMesh.WedgeTangentZ.Num() == NumWedges);
 
-			// Generate lightmap UVs
-			if( SrcModel.BuildSettings.bGenerateLightmapUVs )
-			{
-				if( RawMesh.WedgeTexCoords[ SrcModel.BuildSettings.SrcLightmapIndex ].Num() == 0 )
+				// Generate lightmap UVs
+				if (SrcModel.BuildSettings.bGenerateLightmapUVs)
 				{
-					SrcModel.BuildSettings.SrcLightmapIndex = 0;
-				}
+					if (RawMesh.WedgeTexCoords[SrcModel.BuildSettings.SrcLightmapIndex].Num() == 0)
+					{
+						SrcModel.BuildSettings.SrcLightmapIndex = 0;
+					}
 
-				FLayoutUV Packer( &RawMesh, SrcModel.BuildSettings.SrcLightmapIndex, SrcModel.BuildSettings.DstLightmapIndex, SrcModel.BuildSettings.MinLightmapResolution );
+					FLayoutUV Packer(&RawMesh, SrcModel.BuildSettings.SrcLightmapIndex, SrcModel.BuildSettings.DstLightmapIndex, SrcModel.BuildSettings.MinLightmapResolution);
 
-				Packer.FindCharts( OverlappingCorners );
-				bool bPackSuccess = Packer.FindBestPacking();
-				if( bPackSuccess )
-				{
-					Packer.CommitPackedUVs();
+					Packer.FindCharts(OverlappingCorners);
+					bool bPackSuccess = Packer.FindBestPacking();
+					if (bPackSuccess)
+					{
+						Packer.CommitPackedUVs();
+					}
 				}
+				HasRawMesh[LODIndex] = true;
+			}
+			else if (LODIndex > 0 && MeshReduction)
+			{
+				// If a raw mesh is not explicitly provided, use the raw mesh of the
+				// next highest LOD.
+				RawMesh = LODMeshes[LODIndex-1];
+				OverlappingCorners = LODOverlappingCorners[LODIndex-1];
+				LODBuildSettings[LODIndex] = LODBuildSettings[LODIndex-1];
+				HasRawMesh[LODIndex] = false;
+			}
+			else
+			{
+				check(0);
 			}
 		}
-		else if (LODIndex > 0 && MeshReduction)
+		check(LODMeshes.Num() == SourceModels.Num());
+		check(LODOverlappingCorners.Num() == SourceModels.Num());
+
+		// Bail if there is no raw mesh data from which to build a renderable mesh.
+		if (LODMeshes.Num() == 0 || LODMeshes[0].WedgeIndices.Num() == 0)
 		{
-			// If a raw mesh is not explicitly provided, use the raw mesh of the
-			// next highest LOD.
-			RawMesh = LODMeshes[LODIndex-1];
-			OverlappingCorners = LODOverlappingCorners[LODIndex-1];
-			LODBuildSettings[LODIndex] = LODBuildSettings[LODIndex-1];
+			return false;
 		}
-	}
-	check(LODMeshes.Num() == SourceModels.Num());
-	check(LODOverlappingCorners.Num() == SourceModels.Num());
 
-	// Bail if there is no raw mesh data from which to build a renderable mesh.
-	if (LODMeshes.Num() == 0 || LODMeshes[0].WedgeIndices.Num() == 0)
-	{
-		return false;
+		Stage = EStage::Gathered;
+		return true;
 	}
 
-	// Reduce each LOD mesh according to its reduction settings.
-	OutRenderData.bReducedBySimplygon = false;
-	int32 NumValidLODs = 0;
-	for (int32 LODIndex = 0; LODIndex < SourceModels.Num(); ++LODIndex)
+	bool ReduceLODs(TArray<FStaticMeshSourceModel>& SourceModels, const FStaticMeshLODGroup& LODGroup, IMeshReduction* MeshReduction, bool& bOutWasReduced)
 	{
-		const FStaticMeshSourceModel& SrcModel = SourceModels[LODIndex];
-		FMeshReductionSettings ReductionSettings = LODGroup.GetSettings(SrcModel.ReductionSettings, LODIndex);
-		LODMaxDeviation[NumValidLODs] = 0.0f;
-		if (LODIndex != NumValidLODs)
-		{
-			LODBuildSettings[NumValidLODs] = LODBuildSettings[LODIndex];
-			LODOverlappingCorners[NumValidLODs] = LODOverlappingCorners[LODIndex];
-		}
+		check(Stage == EStage::Gathered);
 
-		if (MeshReduction && (ReductionSettings.PercentTriangles < 1.0f || ReductionSettings.MaxDeviation > 0.0f))
+		// Reduce each LOD mesh according to its reduction settings.
+		for (int32 LODIndex = 0; LODIndex < SourceModels.Num(); ++LODIndex)
 		{
-			FRawMesh InMesh = LODMeshes[ReductionSettings.BaseLODModel];
-			FRawMesh& DestMesh = LODMeshes[NumValidLODs];
-			TMultiMap<int32,int32>& DestOverlappingCorners = LODOverlappingCorners[NumValidLODs];
+			const FStaticMeshSourceModel& SrcModel = SourceModels[LODIndex];
+			FMeshReductionSettings ReductionSettings = LODGroup.GetSettings(SrcModel.ReductionSettings, LODIndex);
+			LODMaxDeviation[NumValidLODs] = 0.0f;
+			if (LODIndex != NumValidLODs)
+			{
+				LODBuildSettings[NumValidLODs] = LODBuildSettings[LODIndex];
+				LODOverlappingCorners[NumValidLODs] = LODOverlappingCorners[LODIndex];
+			}
 
-			MeshReduction->Reduce(DestMesh, LODMaxDeviation[NumValidLODs], InMesh, ReductionSettings);
-			if (DestMesh.WedgeIndices.Num() > 0 && !DestMesh.IsValid())
+			if (MeshReduction && (ReductionSettings.PercentTriangles < 1.0f || ReductionSettings.MaxDeviation > 0.0f))
 			{
-				UE_LOG(LogMeshUtilities,Error,TEXT("Mesh reduction produced a corrupt mesh for LOD%d"),LODIndex);
-				return false;
+				FRawMesh InMesh = LODMeshes[ReductionSettings.BaseLODModel];
+				FRawMesh& DestMesh = LODMeshes[NumValidLODs];
+				TMultiMap<int32, int32>& DestOverlappingCorners = LODOverlappingCorners[NumValidLODs];
+
+				MeshReduction->Reduce(DestMesh, LODMaxDeviation[NumValidLODs], InMesh, ReductionSettings);
+				if (DestMesh.WedgeIndices.Num() > 0 && !DestMesh.IsValid())
+				{
+					UE_LOG(LogMeshUtilities, Error, TEXT("Mesh reduction produced a corrupt mesh for LOD%d"), LODIndex);
+					return false;
+				}
+				bOutWasReduced = true;
+
+				// Recompute adjacency information.
+				DestOverlappingCorners.Reset();
+				float ComparisonThreshold = GetComparisonThreshold(LODBuildSettings[NumValidLODs]);
+				FindOverlappingCorners(DestOverlappingCorners, DestMesh, ComparisonThreshold);
 			}
-			OutRenderData.bReducedBySimplygon = bUsingSimplygon;
 
-			// Recompute adjacency information.
-			DestOverlappingCorners.Reset();
-			float ComparisonThreshold = GetComparisonThreshold(LODBuildSettings[NumValidLODs]);
-			FindOverlappingCorners(DestOverlappingCorners, DestMesh, ComparisonThreshold);
+			if (LODMeshes[NumValidLODs].WedgeIndices.Num() > 0)
+			{
+				NumValidLODs++;
+			}
 		}
 
-		if (LODMeshes[NumValidLODs].WedgeIndices.Num() > 0)
+		if (NumValidLODs < 1)
 		{
-			NumValidLODs++;
+			return false;
 		}
+		Stage = EStage::Reduce;
+		return true;
 	}
 
-	if (NumValidLODs < 1)
-	{
-		return false;
-	}
-
-	// Generate per-LOD rendering data.
-	OutRenderData.AllocateLODResources(NumValidLODs);
-	for (int32 LODIndex = 0; LODIndex < NumValidLODs; ++LODIndex)
+	bool GenerateRenderingMeshes(FMeshUtilities& MeshUtilities, FStaticMeshRenderData& OutRenderData, TArray<FStaticMeshSourceModel>& InOutModels)
 	{
-		FStaticMeshLODResources& LODModel = OutRenderData.LODResources[LODIndex];
-		FRawMesh& RawMesh = LODMeshes[LODIndex];
-		LODModel.MaxDeviation = LODMaxDeviation[LODIndex];
+		check(Stage == EStage::Reduce);
+		// Generate per-LOD rendering data.
+		OutRenderData.AllocateLODResources(NumValidLODs);
+		for (int32 LODIndex = 0; LODIndex < NumValidLODs; ++LODIndex)
+		{
+			FStaticMeshLODResources& LODModel = OutRenderData.LODResources[LODIndex];
+			FRawMesh& RawMesh = LODMeshes[LODIndex];
+			LODModel.MaxDeviation = LODMaxDeviation[LODIndex];
 
-		TArray<FStaticMeshBuildVertex> Vertices;
-		TArray<TArray<uint32> > PerSectionIndices;
+			TArray<FStaticMeshBuildVertex> Vertices;
+			TArray<TArray<uint32> > PerSectionIndices;
 
-		// Find out how many sections are in the mesh.
-		int32 MaxMaterialIndex = -1;
-		for (int32 FaceIndex = 0; FaceIndex < RawMesh.FaceMaterialIndices.Num(); FaceIndex++)
-		{
-			MaxMaterialIndex = FMath::Max<int32>(RawMesh.FaceMaterialIndices[FaceIndex],MaxMaterialIndex);
-		}
-		MaxMaterialIndex = FMath::Min(MaxMaterialIndex, MAX_MESH_MATERIAL_INDEX);
+			// Find out how many sections are in the mesh.
+			int32 MaxMaterialIndex = -1;
+			for (int32 FaceIndex = 0; FaceIndex < RawMesh.FaceMaterialIndices.Num(); FaceIndex++)
+			{
+				MaxMaterialIndex = FMath::Max<int32>(RawMesh.FaceMaterialIndices[FaceIndex], MaxMaterialIndex);
+			}
+			MaxMaterialIndex = FMath::Min(MaxMaterialIndex, MAX_MESH_MATERIAL_INDEX);
 
-		while (MaxMaterialIndex >= LODModel.Sections.Num())
-		{
-			FStaticMeshSection* Section = new(LODModel.Sections) FStaticMeshSection();
-			Section->MaterialIndex = LODModel.Sections.Num() - 1;
-			new(PerSectionIndices) TArray<uint32>;
-		}
+			while (MaxMaterialIndex >= LODModel.Sections.Num())
+			{
+				FStaticMeshSection* Section = new(LODModel.Sections) FStaticMeshSection();
+				Section->MaterialIndex = LODModel.Sections.Num() - 1;
+				new(PerSectionIndices)TArray<uint32>;
+			}
 
-		// Build and cache optimize vertex and index buffers.
-		{
-			// TODO_STATICMESH: The wedge map is only valid for LODIndex 0 if no reduction has been performed.
-			// We can compute an approximate one instead for other LODs.
-			TArray<int32> TempWedgeMap;
-			TArray<int32>& WedgeMap = (LODIndex == 0 && SourceModels[0].ReductionSettings.PercentTriangles >= 1.0f) ? OutRenderData.WedgeMap : TempWedgeMap;
-			float ComparisonThreshold = GetComparisonThreshold(LODBuildSettings[LODIndex]);
-			BuildStaticMeshVertexAndIndexBuffers(Vertices, PerSectionIndices, WedgeMap, RawMesh, LODOverlappingCorners[LODIndex], ComparisonThreshold, LODBuildSettings[LODIndex].BuildScale3D );
-			check(WedgeMap.Num() == RawMesh.WedgeIndices.Num());
-
-			if( RawMesh.WedgeIndices.Num() < 50000 * 3 )
+			// Build and cache optimize vertex and index buffers.
 			{
-				CacheOptimizeVertexAndIndexBuffer(Vertices, PerSectionIndices, WedgeMap);
+				// TODO_STATICMESH: The wedge map is only valid for LODIndex 0 if no reduction has been performed.
+				// We can compute an approximate one instead for other LODs.
+				TArray<int32> TempWedgeMap;
+				TArray<int32>& WedgeMap = (LODIndex == 0 && InOutModels[0].ReductionSettings.PercentTriangles >= 1.0f) ? OutRenderData.WedgeMap : TempWedgeMap;
+				float ComparisonThreshold = GetComparisonThreshold(LODBuildSettings[LODIndex]);
+				BuildStaticMeshVertexAndIndexBuffers(Vertices, PerSectionIndices, WedgeMap, RawMesh, LODOverlappingCorners[LODIndex], ComparisonThreshold, LODBuildSettings[LODIndex].BuildScale3D);
 				check(WedgeMap.Num() == RawMesh.WedgeIndices.Num());
-			}
-		}
 
-		// Initialize the vertex buffer.
-		int32 NumTexCoords = ComputeNumTexCoords(RawMesh, MAX_STATIC_TEXCOORDS);
-		LODModel.VertexBuffer.SetUseFullPrecisionUVs(LODBuildSettings[LODIndex].bUseFullPrecisionUVs);	
-		LODModel.VertexBuffer.Init(Vertices,NumTexCoords);
-		LODModel.PositionVertexBuffer.Init(Vertices);
-		LODModel.ColorVertexBuffer.Init(Vertices);
+				if (RawMesh.WedgeIndices.Num() < 50000 * 3)
+				{
+					MeshUtilities.CacheOptimizeVertexAndIndexBuffer(Vertices, PerSectionIndices, WedgeMap);
+					check(WedgeMap.Num() == RawMesh.WedgeIndices.Num());
+				}
+			}
 
-		// Concatenate the per-section index buffers.
-		TArray<uint32> CombinedIndices;
-		bool bNeeds32BitIndices = false;
-		for (int32 SectionIndex = 0; SectionIndex < LODModel.Sections.Num(); SectionIndex++)
-		{
-			FStaticMeshSection& Section = LODModel.Sections[SectionIndex];
-			TArray<uint32> const& SectionIndices = PerSectionIndices[SectionIndex];
-			Section.FirstIndex = 0;
-			Section.NumTriangles = 0;
-			Section.MinVertexIndex = 0;
-			Section.MaxVertexIndex = 0;
-
-			if (SectionIndices.Num())
+			// Initialize the vertex buffer.
+			int32 NumTexCoords = ComputeNumTexCoords(RawMesh, MAX_STATIC_TEXCOORDS);
+			LODModel.VertexBuffer.SetUseFullPrecisionUVs(LODBuildSettings[LODIndex].bUseFullPrecisionUVs);
+			LODModel.VertexBuffer.Init(Vertices, NumTexCoords);
+			LODModel.PositionVertexBuffer.Init(Vertices);
+			LODModel.ColorVertexBuffer.Init(Vertices);
+
+			// Concatenate the per-section index buffers.
+			TArray<uint32> CombinedIndices;
+			bool bNeeds32BitIndices = false;
+			for (int32 SectionIndex = 0; SectionIndex < LODModel.Sections.Num(); SectionIndex++)
 			{
-				Section.FirstIndex = CombinedIndices.Num();
-				Section.NumTriangles = SectionIndices.Num() / 3;
+				FStaticMeshSection& Section = LODModel.Sections[SectionIndex];
+				TArray<uint32> const& SectionIndices = PerSectionIndices[SectionIndex];
+				Section.FirstIndex = 0;
+				Section.NumTriangles = 0;
+				Section.MinVertexIndex = 0;
+				Section.MaxVertexIndex = 0;
+
+				if (SectionIndices.Num())
+				{
+					Section.FirstIndex = CombinedIndices.Num();
+					Section.NumTriangles = SectionIndices.Num() / 3;
+
+					CombinedIndices.AddUninitialized(SectionIndices.Num());
+					uint32* DestPtr = &CombinedIndices[Section.FirstIndex];
+					uint32 const* SrcPtr = SectionIndices.GetData();
 
-				CombinedIndices.AddUninitialized(SectionIndices.Num());
-				uint32* DestPtr = &CombinedIndices[Section.FirstIndex];
-				uint32 const* SrcPtr = SectionIndices.GetData();
+					Section.MinVertexIndex = *SrcPtr;
+					Section.MaxVertexIndex = *SrcPtr;
+
+					for (int32 Index = 0; Index < SectionIndices.Num(); Index++)
+					{
+						uint32 VertIndex = *SrcPtr++;
+
+						bNeeds32BitIndices |= (VertIndex > MAX_uint16);
+						Section.MinVertexIndex = FMath::Min<uint32>(VertIndex, Section.MinVertexIndex);
+						Section.MaxVertexIndex = FMath::Max<uint32>(VertIndex, Section.MaxVertexIndex);
+						*DestPtr++ = VertIndex;
+					}
+				}
+			}
+			LODModel.IndexBuffer.SetIndices(CombinedIndices, bNeeds32BitIndices ? EIndexBufferStride::Force32Bit : EIndexBufferStride::Force16Bit);
 
-				Section.MinVertexIndex = *SrcPtr;
-				Section.MaxVertexIndex = *SrcPtr;
+			if (LODIndex == 0)
+			{
+				ComputeStreamingTextureFactors(
+					OutRenderData.StreamingTextureFactors,
+					&OutRenderData.MaxStreamingTextureFactor,
+					RawMesh,
+					LODBuildSettings[LODIndex].BuildScale3D
+					);
+			}
 
-				for (int32 Index = 0; Index < SectionIndices.Num(); Index++)
+			// Build the depth-only index buffer.
+			{
+				TArray<uint32> DepthOnlyIndices;
+				BuildDepthOnlyIndexBuffer(
+					DepthOnlyIndices,
+					Vertices,
+					CombinedIndices,
+					LODModel.Sections
+					);
+
+				if (DepthOnlyIndices.Num() < 50000 * 3)
 				{
-					uint32 VertIndex = *SrcPtr++;
+					MeshUtilities.CacheOptimizeIndexBuffer(DepthOnlyIndices);
+				}
+
+				LODModel.DepthOnlyIndexBuffer.SetIndices(DepthOnlyIndices, bNeeds32BitIndices ? EIndexBufferStride::Force32Bit : EIndexBufferStride::Force16Bit);
+			}
+
+			// Build a list of wireframe edges in the static mesh.
+			{
+				TArray<FMeshEdge> Edges;
+				TArray<uint32> WireframeIndices;
 
-					bNeeds32BitIndices |= (VertIndex > MAX_uint16);
-					Section.MinVertexIndex = FMath::Min<uint32>(VertIndex,Section.MinVertexIndex);
-					Section.MaxVertexIndex = FMath::Max<uint32>(VertIndex,Section.MaxVertexIndex);
-					*DestPtr++ = VertIndex;
+				FStaticMeshEdgeBuilder(CombinedIndices, Vertices, Edges).FindEdges();
+				WireframeIndices.Empty(2 * Edges.Num());
+				for (int32 EdgeIndex = 0;EdgeIndex < Edges.Num();EdgeIndex++)
+				{
+					FMeshEdge&	Edge = Edges[EdgeIndex];
+					WireframeIndices.Add(Edge.Vertices[0]);
+					WireframeIndices.Add(Edge.Vertices[1]);
 				}
+				LODModel.WireframeIndexBuffer.SetIndices(WireframeIndices, bNeeds32BitIndices ? EIndexBufferStride::Force32Bit : EIndexBufferStride::Force16Bit);
+			}
+
+			// Build the adjacency index buffer used for tessellation.
+			{
+				TArray<uint32> AdjacencyIndices;
+
+				BuildStaticAdjacencyIndexBuffer(
+					LODModel.PositionVertexBuffer,
+					LODModel.VertexBuffer,
+					CombinedIndices,
+					AdjacencyIndices
+					);
+				LODModel.AdjacencyIndexBuffer.SetIndices(AdjacencyIndices, bNeeds32BitIndices ? EIndexBufferStride::Force32Bit : EIndexBufferStride::Force16Bit);
 			}
 		}
-		LODModel.IndexBuffer.SetIndices(CombinedIndices, bNeeds32BitIndices ? EIndexBufferStride::Force32Bit : EIndexBufferStride::Force16Bit);
 
-		if (LODIndex == 0)
+		// Copy the original material indices to fixup meshes before compacting of materials was done.
+		if (NumValidLODs > 0)
 		{
-			ComputeStreamingTextureFactors(
-				OutRenderData.StreamingTextureFactors,
-				&OutRenderData.MaxStreamingTextureFactor,
-				RawMesh,
-				LODBuildSettings[LODIndex].BuildScale3D
-				);
+			OutRenderData.MaterialIndexToImportIndex = LODMeshes[0].MaterialIndexToImportIndex;
 		}
 
-		// Build the depth-only index buffer.
+		// Calculate the bounding box.
+		FBox BoundingBox(0);
+		FPositionVertexBuffer& BasePositionVertexBuffer = OutRenderData.LODResources[0].PositionVertexBuffer;
+		for (uint32 VertexIndex = 0; VertexIndex < BasePositionVertexBuffer.GetNumVertices(); VertexIndex++)
 		{
-			TArray<uint32> DepthOnlyIndices;
-			BuildDepthOnlyIndexBuffer(
-				DepthOnlyIndices,
-				Vertices,
-				CombinedIndices,
-				LODModel.Sections
+			BoundingBox += BasePositionVertexBuffer.VertexPosition(VertexIndex);
+		}
+		BoundingBox.GetCenterAndExtents(OutRenderData.Bounds.Origin, OutRenderData.Bounds.BoxExtent);
+
+		// Calculate the bounding sphere, using the center of the bounding box as the origin.
+		OutRenderData.Bounds.SphereRadius = 0.0f;
+		for (uint32 VertexIndex = 0; VertexIndex < BasePositionVertexBuffer.GetNumVertices(); VertexIndex++)
+		{
+			OutRenderData.Bounds.SphereRadius = FMath::Max(
+				(BasePositionVertexBuffer.VertexPosition(VertexIndex) - OutRenderData.Bounds.Origin).Size(),
+				OutRenderData.Bounds.SphereRadius
 				);
+		}
 
-			if( DepthOnlyIndices.Num() < 50000 * 3 )
-			{
-				CacheOptimizeIndexBuffer(DepthOnlyIndices);
-			}
+		Stage = EStage::GenerateRendering;
+		return true;
+	}
 
-			LODModel.DepthOnlyIndexBuffer.SetIndices(DepthOnlyIndices, bNeeds32BitIndices ? EIndexBufferStride::Force32Bit : EIndexBufferStride::Force16Bit);
-		}
+	bool ReplaceRawMeshModels(TArray<FStaticMeshSourceModel>& SourceModels)
+	{
+		check(Stage == EStage::Reduce);
 
-		// Build a list of wireframe edges in the static mesh.
+		check(HasRawMesh[0]);
+		check(SourceModels.Num() >= NumValidLODs);
+		bool bDirty = false;
+		for (int32 Index = 1; Index < NumValidLODs; ++Index)
 		{
-			TArray<FMeshEdge> Edges;
-			TArray<uint32> WireframeIndices;
-
-			FStaticMeshEdgeBuilder(CombinedIndices,Vertices,Edges).FindEdges();
-			WireframeIndices.Empty(2 * Edges.Num());
-			for(int32 EdgeIndex = 0;EdgeIndex < Edges.Num();EdgeIndex++)
+			if (!HasRawMesh[Index])
 			{
-				FMeshEdge&	Edge = Edges[EdgeIndex];
-				WireframeIndices.Add(Edge.Vertices[0]);
-				WireframeIndices.Add(Edge.Vertices[1]);
+				SourceModels[Index].RawMeshBulkData->SaveRawMesh(LODMeshes[Index]);
+				bDirty = true;
 			}
-			LODModel.WireframeIndexBuffer.SetIndices(WireframeIndices, bNeeds32BitIndices ? EIndexBufferStride::Force32Bit : EIndexBufferStride::Force16Bit);
 		}
 
-		// Build the adjacency index buffer used for tessellation.
-		{
-			TArray<uint32> AdjacencyIndices;
+		Stage = EStage::ReplaceRaw;
+		return true;
+	}
 
-			BuildStaticAdjacencyIndexBuffer(
-				LODModel.PositionVertexBuffer,
-				LODModel.VertexBuffer,
-				CombinedIndices,
-				AdjacencyIndices
-				);
-			LODModel.AdjacencyIndexBuffer.SetIndices(AdjacencyIndices, bNeeds32BitIndices ? EIndexBufferStride::Force32Bit : EIndexBufferStride::Force16Bit);
-		}
+private:
+	enum class EStage
+	{
+		Uninit,
+		Gathered,
+		Reduce,
+		GenerateRendering,
+		ReplaceRaw,
+	};
+
+	EStage Stage;
+
+	int32 NumValidLODs;
+
+	TIndirectArray<FRawMesh> LODMeshes;
+	TIndirectArray<TMultiMap<int32, int32> > LODOverlappingCorners;
+	float LODMaxDeviation[MAX_STATIC_MESH_LODS];
+	FMeshBuildSettings LODBuildSettings[MAX_STATIC_MESH_LODS];
+	bool HasRawMesh[MAX_STATIC_MESH_LODS];
+};
+
+bool FMeshUtilities::BuildStaticMesh(FStaticMeshRenderData& OutRenderData, TArray<FStaticMeshSourceModel>& SourceModels, const FStaticMeshLODGroup& LODGroup)
+{
+	FStaticMeshUtilityBuilder Builder;
+	if (!Builder.GatherSourceMeshesPerLOD(SourceModels, MeshReduction))
+	{
+		return false;
 	}
 
-	// Copy the original material indices to fixup meshes before compacting of materials was done.
-	if (NumValidLODs > 0)
+	OutRenderData.bReducedBySimplygon = false;
+	bool bWasReduced = false;
+	if (!Builder.ReduceLODs(SourceModels, LODGroup, MeshReduction, bWasReduced))
 	{
-		OutRenderData.MaterialIndexToImportIndex = LODMeshes[0].MaterialIndexToImportIndex;
+		return false;
 	}
+	OutRenderData.bReducedBySimplygon = (bWasReduced && bUsingSimplygon);
 
-	// Calculate the bounding box.
-	FBox BoundingBox(0);
-	FPositionVertexBuffer& BasePositionVertexBuffer = OutRenderData.LODResources[0].PositionVertexBuffer;
-	for (uint32 VertexIndex = 0; VertexIndex < BasePositionVertexBuffer.GetNumVertices(); VertexIndex++)
+	return Builder.GenerateRenderingMeshes(*this, OutRenderData, SourceModels);
+}
+
+bool FMeshUtilities::GenerateStaticMeshLODs(TArray<FStaticMeshSourceModel>& Models, const FStaticMeshLODGroup& LODGroup)
+{
+	FStaticMeshUtilityBuilder Builder;
+	if (!Builder.GatherSourceMeshesPerLOD(Models, MeshReduction))
 	{
-		BoundingBox += BasePositionVertexBuffer.VertexPosition(VertexIndex);
+		return false;
 	}
-	BoundingBox.GetCenterAndExtents(OutRenderData.Bounds.Origin,OutRenderData.Bounds.BoxExtent);
 
-	// Calculate the bounding sphere, using the center of the bounding box as the origin.
-	OutRenderData.Bounds.SphereRadius = 0.0f;
-	for (uint32 VertexIndex = 0; VertexIndex < BasePositionVertexBuffer.GetNumVertices(); VertexIndex++)
+	bool bWasReduced = false;
+	if (!Builder.ReduceLODs(Models, LODGroup, MeshReduction, bWasReduced))
 	{
-		OutRenderData.Bounds.SphereRadius = FMath::Max(
-			(BasePositionVertexBuffer.VertexPosition(VertexIndex) - OutRenderData.Bounds.Origin).Size(),
-			OutRenderData.Bounds.SphereRadius
-			);
+		return false;
 	}
 
-	return true;
-}
+	if (bWasReduced)
+	{
+		return Builder.ReplaceRawMeshModels(Models);
+	}
 
+	return false;
+}
 
 //@TODO: The OutMessages has to be a struct that contains FText/FName, or make it Token and add that as error. Needs re-work. Temporary workaround for now. 
 bool FMeshUtilities::BuildSkeletalMesh( FStaticLODModel& LODModel, const FReferenceSkeleton& RefSkeleton, const TArray<FVertInfluence>& Influences, const TArray<FMeshWedge>& Wedges, const TArray<FMeshFace>& Faces, const TArray<FVector>& Points, const TArray<int32>& PointToOriginalMap, bool bKeepOverlappingVertices, bool bComputeNormals, bool bComputeTangents, TArray<FText> * OutWarningMessages, TArray<FName> * OutWarningNames)
diff --git a/Engine/Source/Developer/MeshUtilities/Public/MeshUtilities.h b/Engine/Source/Developer/MeshUtilities/Public/MeshUtilities.h
--- a/Engine/Source/Developer/MeshUtilities/Public/MeshUtilities.h
+++ b/Engine/Source/Developer/MeshUtilities/Public/MeshUtilities.h
@@ -134,6 +134,16 @@ public:
 		const class FStaticMeshLODGroup& LODGroup
 		) = 0;
 
+	/**
+	 * Builds a static mesh using the provided source models and the LOD groups settings, and replaces
+	 * the RawMeshes with the reduced meshes. Does not modify renderable data.
+	 * @returns true if the meshes were built successfully.
+	 */
+	virtual bool GenerateStaticMeshLODs(
+		TArray<struct FStaticMeshSourceModel>& Models,
+		const class FStaticMeshLODGroup& LODGroup
+		) = 0;
+
 	/** Builds a signed distance field volume for the given LODModel. */
 	virtual void GenerateSignedDistanceFieldVolumeData(
 		const FStaticMeshLODResources& LODModel,
diff --git a/Engine/Source/Editor/StaticMeshEditor/Private/StaticMeshEditor.cpp b/Engine/Source/Editor/StaticMeshEditor/Private/StaticMeshEditor.cpp
--- a/Engine/Source/Editor/StaticMeshEditor/Private/StaticMeshEditor.cpp
+++ b/Engine/Source/Editor/StaticMeshEditor/Private/StaticMeshEditor.cpp
@@ -35,6 +35,17 @@
 
 DEFINE_LOG_CATEGORY_STATIC(LogStaticMeshEditor, Log, All);
 
+class FStaticMeshStatusMessageContext : public FScopedSlowTask
+{
+public:
+	explicit FStaticMeshStatusMessageContext(const FText& InMessage)
+		: FScopedSlowTask(0, InMessage)
+	{
+		UE_LOG(LogStaticMesh, Log, TEXT("%s"), *InMessage.ToString());
+		MakeDialog();
+	}
+};
+
 const FName FStaticMeshEditor::ViewportTabId( TEXT( "StaticMeshEditor_Viewport" ) );
 const FName FStaticMeshEditor::PropertiesTabId( TEXT( "StaticMeshEditor_Properties" ) );
 const FName FStaticMeshEditor::SocketManagerTabId( TEXT( "StaticMeshEditor_SocketManager" ) );
@@ -211,6 +222,11 @@ void FStaticMeshEditor::ExtendMenu()
 			InMenuBuilder.BeginSection("MeshChange");
 			{
 				InMenuBuilder.AddMenuEntry(FStaticMeshEditorCommands::Get().ChangeMesh);
+				static auto* CVar = IConsoleManager::Get().FindTConsoleVariableDataInt(TEXT("r.StaticMesh.EnableSaveGeneratedLODsInPackage"));
+				if (CVar && CVar->GetValueOnGameThread() != 0)
+				{
+					InMenuBuilder.AddMenuEntry(FStaticMeshEditorCommands::Get().SaveGeneratedLODs);
+				}
 			}
 			InMenuBuilder.EndSection();
 		}
@@ -422,6 +438,10 @@ void FStaticMeshEditor::BindCommands()
 		FExecuteAction::CreateSP(this, &FStaticMeshEditor::OnChangeMesh),
 		FCanExecuteAction::CreateSP(this, &FStaticMeshEditor::CanChangeMesh));
 
+	UICommandList->MapAction(
+		Commands.SaveGeneratedLODs,
+		FExecuteAction::CreateSP(this, &FStaticMeshEditor::OnSaveGeneratedLODs));
+
 	// Collision Menu
 	UICommandList->MapAction(
 		Commands.CreateAutoConvexCollision,
@@ -1602,6 +1622,21 @@ bool FStaticMeshEditor::CanChangeMesh() const
 	return CanChange;
 }
 
+void FStaticMeshEditor::OnSaveGeneratedLODs()
+{
+	if (StaticMesh)
+	{
+		StaticMesh->GenerateLodsInPackage();
+
+		// Update editor UI as we modified LOD groups
+		auto Selected = StaticMeshDetailsView->GetSelectedObjects();
+		StaticMeshDetailsView->SetObjects(Selected, true);
+
+		// Update screen
+		Viewport->RefreshViewport();
+	}
+}
+
 void FStaticMeshEditor::DoDecomp(int32 InMaxHullCount, int32 InMaxHullVerts)
 {
 	// Check we have a selected StaticMesh
diff --git a/Engine/Source/Editor/StaticMeshEditor/Private/StaticMeshEditor.h b/Engine/Source/Editor/StaticMeshEditor/Private/StaticMeshEditor.h
--- a/Engine/Source/Editor/StaticMeshEditor/Private/StaticMeshEditor.h
+++ b/Engine/Source/Editor/StaticMeshEditor/Private/StaticMeshEditor.h
@@ -197,6 +197,9 @@ private:
 	/** Whether there is a static mesh selected in the content browser to change to*/
 	bool CanChangeMesh() const;
 
+	/** Replace the generated LODs in the original source mesh with the reduced versions.*/
+	void OnSaveGeneratedLODs();
+
 	/** Rebuilds the LOD combo list and sets it to "auto", a safe LOD level. */
 	void RegenerateLODComboList();
 
diff --git a/Engine/Source/Editor/StaticMeshEditor/Private/StaticMeshEditorActions.cpp b/Engine/Source/Editor/StaticMeshEditor/Private/StaticMeshEditorActions.cpp
--- a/Engine/Source/Editor/StaticMeshEditor/Private/StaticMeshEditorActions.cpp
+++ b/Engine/Source/Editor/StaticMeshEditor/Private/StaticMeshEditorActions.cpp
@@ -41,5 +41,7 @@ void FStaticMeshEditorCommands::RegisterCommands()
 	UI_COMMAND( FindSource, "Find Source", "Opens explorer at the location of this asset.", EUserInterfaceActionType::Button, FInputGesture() );
 	
 	UI_COMMAND( ChangeMesh, "Change Mesh", "Changes the static mesh asset loaded in the Static Mesh Editor to the asset currently selected in the Content Browser.", EUserInterfaceActionType::Button, FInputGesture() );
+
+	UI_COMMAND( SaveGeneratedLODs, "Save Generated LODs", "Saves the generated LODs in the mesh package.", EUserInterfaceActionType::Button, FInputGesture());
 }
 
diff --git a/Engine/Source/Editor/StaticMeshEditor/Public/StaticMeshEditorActions.h b/Engine/Source/Editor/StaticMeshEditor/Public/StaticMeshEditorActions.h
--- a/Engine/Source/Editor/StaticMeshEditor/Public/StaticMeshEditorActions.h
+++ b/Engine/Source/Editor/StaticMeshEditor/Public/StaticMeshEditorActions.h
@@ -62,6 +62,8 @@ public:
 
 	TSharedPtr< FUICommandInfo > ChangeMesh;
 
+	TSharedPtr< FUICommandInfo > SaveGeneratedLODs;
+
 	/**
 	 * Initialize commands
 	 */
diff --git a/Engine/Source/Runtime/Engine/Classes/Engine/StaticMesh.h b/Engine/Source/Runtime/Engine/Classes/Engine/StaticMesh.h
--- a/Engine/Source/Runtime/Engine/Classes/Engine/StaticMesh.h
+++ b/Engine/Source/Runtime/Engine/Classes/Engine/StaticMesh.h
@@ -607,6 +607,8 @@ public:
 	 */
 	ENGINE_API static void GetLODGroupsDisplayNames(TArray<FText>& OutLODGroupsDisplayNames);
 
+	ENGINE_API void GenerateLodsInPackage();
+
 private:
 	/**
 	 * Converts legacy LODDistance in the source models to Display Factor
diff --git a/Engine/Source/Runtime/Engine/Private/StaticMesh.cpp b/Engine/Source/Runtime/Engine/Private/StaticMesh.cpp
--- a/Engine/Source/Runtime/Engine/Private/StaticMesh.cpp
+++ b/Engine/Source/Runtime/Engine/Private/StaticMesh.cpp
@@ -1089,7 +1089,7 @@ namespace StaticMeshDerivedDataTimings
 		);
 }
 
-FString BuildStaticMeshDerivedDataKey(UStaticMesh* Mesh, const FStaticMeshLODGroup& LODGroup)
+static FString BuildStaticMeshDerivedDataKey(UStaticMesh* Mesh, const FStaticMeshLODGroup& LODGroup)
 {
 	FString KeySuffix(TEXT(""));
 	TArray<uint8> TempBytes;
@@ -1127,7 +1127,7 @@ FString BuildStaticMeshDerivedDataKey(UStaticMesh* Mesh, const FStaticMeshLODGro
 		);
 }
 
-FString BuildDistanceFieldDerivedDataKey(const FString& InMeshKey)
+static FString BuildDistanceFieldDerivedDataKey(const FString& InMeshKey)
 {
 	return FDerivedDataCacheInterface::BuildCacheKey(
 		TEXT("DIST"),
@@ -2545,7 +2545,38 @@ void UStaticMesh::ConvertLegacyLODDistance()
 	}
 }
 
-#endif // #if WITH_EDITORONLY_DATA
+void UStaticMesh::GenerateLodsInPackage()
+{
+	FFormatNamedArguments Args;
+	Args.Add(TEXT("StaticMeshName"), FText::FromString(GetName()));
+	FStaticMeshStatusMessageContext StatusContext(FText::Format(NSLOCTEXT("Engine", "SavingStaticMeshLODsStatus", "Saving generated LODs for static mesh {StaticMeshName}..."), Args));
+
+	// Get LODGroup info
+	ITargetPlatformManagerModule& TargetPlatformManager = GetTargetPlatformManagerRef();
+	ITargetPlatform* RunningPlatform = TargetPlatformManager.GetRunningTargetPlatform();
+	check(RunningPlatform);
+	const FStaticMeshLODSettings& LODSettings = RunningPlatform->GetStaticMeshLODSettings();
+
+	// Generate the reduced models
+	IMeshUtilities& MeshUtilities = FModuleManager::Get().LoadModuleChecked<IMeshUtilities>(TEXT("MeshUtilities"));
+	if (MeshUtilities.GenerateStaticMeshLODs(SourceModels, LODSettings.GetLODGroup(LODGroup)))
+	{
+		// Clear LOD settings
+		LODGroup = NAME_None;
+		const auto& NewGroup = LODSettings.GetLODGroup(LODGroup);
+		for (int32 Index = 0; Index < SourceModels.Num(); ++Index)
+		{
+			SourceModels[Index].ReductionSettings = NewGroup.GetDefaultSettings(0);
+		}
+
+		Build(true);
+
+		// Raw mesh is now dirty, so the package has to be resaved
+		MarkPackageDirty();
+	}
+}
+
+#endif // #if WITH_EDITOR
 
 UStaticMeshSocket* UStaticMesh::FindSocket(FName InSocketName)
 {
