UE4 - Fix crash when Simplygon is not installed

diff --git a/Engine/Source/Developer/MeshUtilities/Private/MeshUtilities.cpp b/Engine/Source/Developer/MeshUtilities/Private/MeshUtilities.cpp
--- a/Engine/Source/Developer/MeshUtilities/Private/MeshUtilities.cpp
+++ b/Engine/Source/Developer/MeshUtilities/Private/MeshUtilities.cpp
@@ -2846,10 +2846,6 @@ public:
 				LODBuildSettings[LODIndex] = LODBuildSettings[LODIndex-1];
 				HasRawMesh[LODIndex] = false;
 			}
-			else
-			{
-				check(0);
-			}
 		}
 		check(LODMeshes.Num() == SourceModels.Num());
 		check(LODOverlappingCorners.Num() == SourceModels.Num());
