Fixed: Merge Actors utility always recomputes normals on source static meshes
UE-11461

diff --git a/Engine/Source/Developer/MeshUtilities/Private/MeshUtilities.cpp b/Engine/Source/Developer/MeshUtilities/Private/MeshUtilities.cpp
--- a/Engine/Source/Developer/MeshUtilities/Private/MeshUtilities.cpp
+++ b/Engine/Source/Developer/MeshUtilities/Private/MeshUtilities.cpp
@@ -3892,17 +3892,53 @@ bool FMeshUtilities::ConstructRawMesh(
 	}
 		
 	int32 NumWedges = OutRawMesh.WedgeIndices.Num();
-	/* Always Recalculate normals, tangents and bitangents */
-	OutRawMesh.WedgeTangentZ.Empty(NumWedges);
-	OutRawMesh.WedgeTangentZ.AddZeroed(NumWedges);
-	OutRawMesh.WedgeTangentX.Empty(NumWedges);
-	OutRawMesh.WedgeTangentX.AddZeroed(NumWedges);
-	OutRawMesh.WedgeTangentY.Empty(NumWedges);
-	OutRawMesh.WedgeTangentY.AddZeroed(NumWedges);
-
-	TMultiMap<int32,int32> OverlappingCorners;
-	FindOverlappingCorners(OverlappingCorners, OutRawMesh, 0.1f);
-	ComputeTangents(OutRawMesh, OverlappingCorners, ETangentOptions::BlendOverlappingNormals);
+
+	// Figure out if we should recompute normals and tangents.
+	bool bRecomputeNormals = SrcModel.BuildSettings.bRecomputeNormals || OutRawMesh.WedgeTangentZ.Num() == 0;
+	bool bRecomputeTangents = SrcModel.BuildSettings.bRecomputeTangents || OutRawMesh.WedgeTangentX.Num() == 0 || OutRawMesh.WedgeTangentY.Num() == 0;
+
+	// Dump normals and tangents if we are recomputing them.
+	if (bRecomputeTangents)
+	{
+		OutRawMesh.WedgeTangentX.Empty(NumWedges);
+		OutRawMesh.WedgeTangentX.AddZeroed(NumWedges);
+		OutRawMesh.WedgeTangentY.Empty(NumWedges);
+		OutRawMesh.WedgeTangentY.AddZeroed(NumWedges);
+	}
+	
+	if (bRecomputeNormals)
+	{
+		OutRawMesh.WedgeTangentZ.Empty(NumWedges);
+		OutRawMesh.WedgeTangentZ.AddZeroed(NumWedges);
+	}
+
+	// Compute any missing tangents.
+	{
+		float ComparisonThreshold = GetComparisonThreshold(SrcModel.BuildSettings);
+		TMultiMap<int32,int32> OverlappingCorners;
+		FindOverlappingCorners(OverlappingCorners, OutRawMesh, ComparisonThreshold);
+
+		// Static meshes always blend normals of overlapping corners.
+		uint32 TangentOptions = ETangentOptions::BlendOverlappingNormals;
+		if (SrcModel.BuildSettings.bRemoveDegenerates)
+		{
+			// If removing degenerate triangles, ignore them when computing tangents.
+			TangentOptions |= ETangentOptions::IgnoreDegenerateTriangles;
+		}
+		if (SrcModel.BuildSettings.bUseMikkTSpace)
+		{
+			ComputeTangents_MikkTSpace(OutRawMesh, OverlappingCorners, TangentOptions);
+		}
+		else
+		{
+			ComputeTangents(OutRawMesh, OverlappingCorners, TangentOptions);
+		}
+	}
+
+	// At this point the mesh will have valid tangents.
+	check(OutRawMesh.WedgeTangentX.Num() == NumWedges);
+	check(OutRawMesh.WedgeTangentY.Num() == NumWedges);
+	check(OutRawMesh.WedgeTangentZ.Num() == NumWedges);
 
 	//Need to store the unique material indices in order to re-map the material indices in each rawmesh
 	//Only using the base mesh
