'Merge Actors' utility merges mesh LOD elements

diff --git a/Engine/Source/Developer/MeshUtilities/Private/MeshUtilities.cpp b/Engine/Source/Developer/MeshUtilities/Private/MeshUtilities.cpp
--- a/Engine/Source/Developer/MeshUtilities/Private/MeshUtilities.cpp
+++ b/Engine/Source/Developer/MeshUtilities/Private/MeshUtilities.cpp
@@ -146,7 +146,8 @@ private:
 		) override;
 	
 	bool ConstructRawMesh(
-		UStaticMeshComponent* MeshComponent, 
+		UStaticMeshComponent* MeshComponent,
+		int32 LODIndex,
 		FRawMesh& OutRawMesh, 
 		TArray<UMaterialInterface*>& OutUniqueMaterials,
 		TArray<int32>& OutGlobalMaterialIndices
@@ -3699,7 +3700,7 @@ void FMeshUtilities::CreateProxyMesh(
 		TArray<int32> RawMeshMaterialMap;
 		int32 RawMeshId = RawMeshes.Add(FRawMesh());
 		
-		if (ConstructRawMesh(MeshComponent, RawMeshes[RawMeshId], StaticMeshMaterials, RawMeshMaterialMap))
+		if (ConstructRawMesh(MeshComponent, 0, RawMeshes[RawMeshId], StaticMeshMaterials, RawMeshMaterialMap))
 		{
 			MaterialMap.Add(RawMeshId, RawMeshMaterialMap);
 			//Store the bounds for each component
@@ -3809,8 +3810,98 @@ void FMeshUtilities::CreateProxyMesh(
 	}
 }
 
+// Exports static mesh LOD render data to a RawMesh
+static void ExportStaticMeshLOD(const FStaticMeshLODResources& StaticMeshLOD, FRawMesh& OutRawMesh)
+{
+	const int32 NumWedges = StaticMeshLOD.IndexBuffer.GetNumIndices();
+	const int32 NumVertexPositions = StaticMeshLOD.PositionVertexBuffer.GetNumVertices();
+	const int32 NumFaces = NumWedges/3;
+
+	// Indices
+	StaticMeshLOD.IndexBuffer.GetCopy(OutRawMesh.WedgeIndices);
+		
+	// Vertex positions
+	if (NumVertexPositions > 0)
+	{
+		OutRawMesh.VertexPositions.Empty(NumVertexPositions);
+		for (int32 PosIdx = 0; PosIdx < NumVertexPositions; ++PosIdx)
+		{
+			FVector Pos = StaticMeshLOD.PositionVertexBuffer.VertexPosition(PosIdx);
+			OutRawMesh.VertexPositions.Add(Pos);
+		}
+	}
+		
+	// Vertex data
+	if (StaticMeshLOD.VertexBuffer.GetNumVertices() > 0)
+	{
+		OutRawMesh.WedgeTangentX.Empty(NumWedges);
+		OutRawMesh.WedgeTangentY.Empty(NumWedges);
+		OutRawMesh.WedgeTangentZ.Empty(NumWedges);
+
+		const int32 NumTexCoords = StaticMeshLOD.VertexBuffer.GetNumTexCoords();
+		for (int32 TexCoodIdx = 0; TexCoodIdx < NumTexCoords; ++TexCoodIdx)
+		{
+			OutRawMesh.WedgeTexCoords[TexCoodIdx].Empty(NumWedges);
+		}
+		
+		for (int32 WedgeIndex : OutRawMesh.WedgeIndices)
+		{
+			FVector WedgeTangentX = StaticMeshLOD.VertexBuffer.VertexTangentX(WedgeIndex);
+			FVector WedgeTangentY = StaticMeshLOD.VertexBuffer.VertexTangentY(WedgeIndex);
+			FVector WedgeTangentZ = StaticMeshLOD.VertexBuffer.VertexTangentZ(WedgeIndex);
+			OutRawMesh.WedgeTangentX.Add(WedgeTangentX);
+			OutRawMesh.WedgeTangentY.Add(WedgeTangentY);
+			OutRawMesh.WedgeTangentZ.Add(WedgeTangentZ);
+
+			for (int32 TexCoodIdx = 0; TexCoodIdx < NumTexCoords; ++TexCoodIdx)
+			{
+				FVector2D WedgeTexCoord = StaticMeshLOD.VertexBuffer.GetVertexUV(WedgeIndex, TexCoodIdx);
+				OutRawMesh.WedgeTexCoords[TexCoodIdx].Add(WedgeTexCoord);
+			}
+		}
+	}
+
+	// Vertex colors
+	if (StaticMeshLOD.ColorVertexBuffer.GetNumVertices() > 0)
+	{
+		OutRawMesh.WedgeColors.Empty(NumWedges);
+		for (int32 WedgeIndex : OutRawMesh.WedgeIndices)
+		{
+			FColor VertexColor = StaticMeshLOD.ColorVertexBuffer.VertexColor(WedgeIndex);
+			OutRawMesh.WedgeColors.Add(VertexColor);
+		}
+	}
+
+	// Materials
+	{
+		OutRawMesh.FaceMaterialIndices.Empty(NumFaces);
+		OutRawMesh.FaceMaterialIndices.SetNumZeroed(NumFaces);
+
+		for (const FStaticMeshSection& Section : StaticMeshLOD.Sections)
+		{
+			uint32 FirstTriangle = Section.FirstIndex/3;
+			for (uint32 TriangleIndex = FirstTriangle; TriangleIndex < Section.NumTriangles; ++TriangleIndex)
+			{
+				OutRawMesh.FaceMaterialIndices[TriangleIndex] = Section.MaterialIndex;
+			}
+		}
+	}
+
+	// Smoothing masks
+	{
+		OutRawMesh.FaceSmoothingMasks.Empty(NumFaces);
+		OutRawMesh.FaceSmoothingMasks.SetNumUninitialized(NumFaces);
+
+		for (auto& SmoothingMask : OutRawMesh.FaceSmoothingMasks)
+		{
+			SmoothingMask = 1;
+		}
+	}
+}
+
 bool FMeshUtilities::ConstructRawMesh(
-	UStaticMeshComponent* InMeshComponent, 
+	UStaticMeshComponent* InMeshComponent,
+	int32 LODIndex,
 	FRawMesh& OutRawMesh, 
 	TArray<UMaterialInterface*>& OutUniqueMaterials,
 	TArray<int32>& OutGlobalMaterialIndices) const
@@ -3822,26 +3913,35 @@ bool FMeshUtilities::ConstructRawMesh(
 		return false;
 	}
 	
-	if (SrcMesh->SourceModels.Num() < 1)
+	if (!SrcMesh->SourceModels.IsValidIndex(LODIndex))
 	{
-		UE_LOG(LogMeshUtilities, Warning, TEXT("No base render mesh found for %s."), *SrcMesh->GetName());
+		UE_LOG(LogMeshUtilities, Log, TEXT("No mesh data found for LOD%d %s."), LODIndex, *SrcMesh->GetName());
 		return false;
 	}
 
-	//Always access the base mesh
-	FStaticMeshSourceModel& SrcModel = SrcMesh->SourceModels[0];
-	if (SrcModel.RawMeshBulkData->IsEmpty())
+	if (!SrcMesh->RenderData->LODResources.IsValidIndex(LODIndex))
 	{
-		UE_LOG(LogMeshUtilities, Error, TEXT("Base render mesh has no imported raw mesh data %s."), *SrcMesh->GetName());
+		UE_LOG(LogMeshUtilities, Warning, TEXT("No mesh render data found for LOD%d %s."), LODIndex, *SrcMesh->GetName());
 		return false;
 	}
+	
+	FStaticMeshSourceModel& SrcModel = SrcMesh->SourceModels[LODIndex];
+	const bool bImportedMesh = !SrcModel.RawMeshBulkData->IsEmpty();
 
-	SrcModel.RawMeshBulkData->LoadRawMesh(OutRawMesh);
-
+	if (bImportedMesh)
+	{
+		SrcModel.RawMeshBulkData->LoadRawMesh(OutRawMesh);
+	}
+	else
+	{
+		// Reconstruct RawMesh from render data
+		ExportStaticMeshLOD(SrcMesh->RenderData->LODResources[LODIndex], OutRawMesh);
+	}
+	
 	// Make sure the raw mesh is not irreparably malformed.
 	if (!OutRawMesh.IsValidOrFixable())
 	{
-		UE_LOG(LogMeshUtilities, Error, TEXT("Raw mesh (%s) is corrupt for LOD%d."), *SrcMesh->GetName(), 1);
+		UE_LOG(LogMeshUtilities, Error, TEXT("Raw mesh (%s) is corrupt for LOD%d."), *SrcMesh->GetName(), LODIndex);
 		return false;
 	}
 	
@@ -3893,9 +3993,12 @@ bool FMeshUtilities::ConstructRawMesh(
 		
 	int32 NumWedges = OutRawMesh.WedgeIndices.Num();
 
-	// Figure out if we should recompute normals and tangents.
-	bool bRecomputeNormals = SrcModel.BuildSettings.bRecomputeNormals || OutRawMesh.WedgeTangentZ.Num() == 0;
-	bool bRecomputeTangents = SrcModel.BuildSettings.bRecomputeTangents || OutRawMesh.WedgeTangentX.Num() == 0 || OutRawMesh.WedgeTangentY.Num() == 0;
+	// Use build settings from base mesh for LOD entries that was generated inside Editor.
+	const FMeshBuildSettings& BuildSettings = bImportedMesh ? SrcModel.BuildSettings : SrcMesh->SourceModels[0].BuildSettings;
+
+	// Figure out if we should recompute normals and tangents. By default generated LODs should not recompute normals
+	bool bRecomputeNormals = (bImportedMesh && BuildSettings.bRecomputeNormals) || OutRawMesh.WedgeTangentZ.Num() == 0;
+	bool bRecomputeTangents = (bImportedMesh && BuildSettings.bRecomputeTangents) || OutRawMesh.WedgeTangentX.Num() == 0 || OutRawMesh.WedgeTangentY.Num() == 0;
 
 	// Dump normals and tangents if we are recomputing them.
 	if (bRecomputeTangents)
@@ -3913,19 +4016,20 @@ bool FMeshUtilities::ConstructRawMesh(
 	}
 
 	// Compute any missing tangents.
+	if (bRecomputeTangents || bRecomputeTangents)
 	{
-		float ComparisonThreshold = GetComparisonThreshold(SrcModel.BuildSettings);
+		float ComparisonThreshold = GetComparisonThreshold(BuildSettings);
 		TMultiMap<int32,int32> OverlappingCorners;
 		FindOverlappingCorners(OverlappingCorners, OutRawMesh, ComparisonThreshold);
 
 		// Static meshes always blend normals of overlapping corners.
 		uint32 TangentOptions = ETangentOptions::BlendOverlappingNormals;
-		if (SrcModel.BuildSettings.bRemoveDegenerates)
+		if (BuildSettings.bRemoveDegenerates)
 		{
 			// If removing degenerate triangles, ignore them when computing tangents.
 			TangentOptions |= ETangentOptions::IgnoreDegenerateTriangles;
 		}
-		if (SrcModel.BuildSettings.bUseMikkTSpace)
+		if (BuildSettings.bUseMikkTSpace)
 		{
 			ComputeTangents_MikkTSpace(OutRawMesh, OverlappingCorners, TangentOptions);
 		}
@@ -3942,7 +4046,7 @@ bool FMeshUtilities::ConstructRawMesh(
 
 	//Need to store the unique material indices in order to re-map the material indices in each rawmesh
 	//Only using the base mesh
-	for (const FStaticMeshSection& Section : SrcMesh->RenderData->LODResources[0].Sections) 
+	for (const FStaticMeshSection& Section : SrcMesh->RenderData->LODResources[LODIndex].Sections) 
 	{
 		// Add material and store the material ID
 		UMaterialInterface* MaterialToAdd = InMeshComponent->GetMaterial(Section.MaterialIndex);
@@ -3958,8 +4062,6 @@ bool FMeshUtilities::ConstructRawMesh(
 	return true;
 }
 
-
-
 /*------------------------------------------------------------------------------
 	Mesh merging 
 ------------------------------------------------------------------------------*/
@@ -4020,69 +4122,93 @@ void FMeshUtilities::MergeActors(
 	FVector& OutMergedActorLocation) const
 {
 	TArray<UStaticMeshComponent*> ComponentsToMerge;
+	ComponentsToMerge.Reserve(SourceActors.Num());
 	
 	// Collect static mesh components
 	for (AActor* Actor : SourceActors)
 	{
 		TInlineComponentArray<UStaticMeshComponent*> Components;
 		Actor->GetComponents<UStaticMeshComponent>(Components);
-		ComponentsToMerge.Append(Components);
+		// Filter out bad components
+		for (UStaticMeshComponent* MeshComponent : Components)
+		{
+			if (MeshComponent->StaticMesh != nullptr && 
+				MeshComponent->StaticMesh->SourceModels.Num() > 0)
+			{
+				ComponentsToMerge.Add(MeshComponent);
+			}
+		}
 	}
-	
+
 	struct FRawMeshExt
 	{
 		FRawMeshExt() 
 		{}
 		
-		FRawMesh	Mesh;
+		FRawMesh	MeshLOD[MAX_STATIC_MESH_LODS];
 		FString		AssetPackageName;
 		FVector		Pivot;	
 	};
 	
+	typedef FIntPoint FMeshIdAndLOD;
+	
 	TArray<UMaterialInterface*>						UniqueMaterials;
-	TMap<int32, TArray<int32>>						MaterialMap;
+	TMap<FMeshIdAndLOD, TArray<int32>>				MaterialMap;
 	TArray<FRawMeshExt>								SourceMeshes;
-	bool											bWithVertexColors = false;
-	bool											bOcuppiedUVChannels[MAX_MESH_TEXTURE_COORDS] = {};
+	bool											bWithVertexColors[MAX_STATIC_MESH_LODS] = {};
+	bool											bOcuppiedUVChannels[MAX_STATIC_MESH_LODS][MAX_MESH_TEXTURE_COORDS] = {};
+	
 	
 	// Convert collected static mesh components into raw meshes
-	SourceMeshes.Empty(ComponentsToMerge.Num());
+	SourceMeshes.Reserve(ComponentsToMerge.Num());
+	SourceMeshes.SetNum(ComponentsToMerge.Num());
 
-	for (UStaticMeshComponent* MeshComponent : ComponentsToMerge)
+	int32 NumMaxLOD = 0;
+	for (int32 MeshId = 0; MeshId < ComponentsToMerge.Num(); ++MeshId)
 	{
-		TArray<int32> MeshMaterialMap;
-		int32 MeshId = SourceMeshes.Add(FRawMeshExt());
-
-		if (ConstructRawMesh(MeshComponent, SourceMeshes[MeshId].Mesh, UniqueMaterials, MeshMaterialMap))
+		UStaticMeshComponent* MeshComponent = ComponentsToMerge[MeshId];
+		// How many LOD entries merged mesh will have
+		NumMaxLOD = FMath::Max(NumMaxLOD, MeshComponent->StaticMesh->SourceModels.Num());
+		// Mesh component pivot point
+		SourceMeshes[MeshId].Pivot = MeshComponent->ComponentToWorld.GetLocation();
+		// Source mesh asset package name
+		SourceMeshes[MeshId].AssetPackageName = MeshComponent->StaticMesh->GetOutermost()->GetName();
+	}
+	
+	NumMaxLOD = FMath::Min(NumMaxLOD, MAX_STATIC_MESH_LODS);
+		
+	for (int32 LODIndex = 0; LODIndex < NumMaxLOD; ++LODIndex)
+	{
+		for (int32 MeshId = 0; MeshId < ComponentsToMerge.Num(); ++MeshId)
 		{
-			MaterialMap.Add(MeshId, MeshMaterialMap);
-
-
-			// Store component location
-			SourceMeshes[MeshId].Pivot = MeshComponent->ComponentToWorld.GetLocation();
+			UStaticMeshComponent* MeshComponent = ComponentsToMerge[MeshId];
 			
-			// Source mesh asset package name
-			SourceMeshes[MeshId].AssetPackageName = MeshComponent->StaticMesh->GetOutermost()->GetName();
+			TArray<int32> MeshMaterialMap;
+			FRawMesh& RawMeshLOD = SourceMeshes[MeshId].MeshLOD[LODIndex];
 
-			// Should we use vertex colors?
-			if (InSettings.bImportVertexColors)
+			// We duplicate lower LOD in case this mesh has no LOD we want
+			int32 ExportLODIndex = FMath::Min(LODIndex, MeshComponent->StaticMesh->SourceModels.Num()-1);
+
+			if (ConstructRawMesh(MeshComponent, ExportLODIndex, RawMeshLOD, UniqueMaterials, MeshMaterialMap))
 			{
-				// Propagate painted vertex colors into our raw mesh
-				PropagatePaintedColorsToRawMesh(MeshComponent, 0, SourceMeshes[MeshId].Mesh);
-				// Whether at least one of the meshes has vertex colors
-				bWithVertexColors|= (SourceMeshes[MeshId].Mesh.WedgeColors.Num() != 0);
-			}
+				MaterialMap.Add(FMeshIdAndLOD(MeshId, LODIndex), MeshMaterialMap);
+				
+				// Should we use vertex colors?
+				if (InSettings.bImportVertexColors)
+				{
+					// Propagate painted vertex colors into our raw mesh
+					PropagatePaintedColorsToRawMesh(MeshComponent, LODIndex, RawMeshLOD);
+					// Whether at least one of the meshes has vertex colors
+					bWithVertexColors[LODIndex]|= (RawMeshLOD.WedgeColors.Num() != 0);
+				}
 			
-			// Which UV channels has data at least in one mesh
-			for (int32 ChannelIdx = 0; ChannelIdx < MAX_MESH_TEXTURE_COORDS; ++ChannelIdx)
-			{
-				bOcuppiedUVChannels[ChannelIdx]|= (SourceMeshes[MeshId].Mesh.WedgeTexCoords[ChannelIdx].Num() != 0);
+				// Which UV channels has data at least in one mesh
+				for (int32 ChannelIdx = 0; ChannelIdx < MAX_MESH_TEXTURE_COORDS; ++ChannelIdx)
+				{
+					bOcuppiedUVChannels[LODIndex][ChannelIdx]|= (RawMeshLOD.WedgeTexCoords[ChannelIdx].Num() != 0);
+				}
 			}
 		}
-		else
-		{
-			SourceMeshes.RemoveAt(MeshId);
-		}
 	}
 
 	if (SourceMeshes.Num() == 0)
@@ -4090,15 +4216,21 @@ void FMeshUtilities::MergeActors(
 		return;	
 	}
 	
-	//For each raw mesh, re-map the material indices according to the MaterialMap
+	// For each raw mesh, re-map the material indices according to the MaterialMap
 	for (int32 MeshIndex = 0; MeshIndex < SourceMeshes.Num(); ++MeshIndex)
 	{
-		FRawMesh& RawMesh = SourceMeshes[MeshIndex].Mesh;
-		const TArray<int32>& Map = *MaterialMap.Find(MeshIndex);
-		for (int32& FaceMaterialIndex : RawMesh.FaceMaterialIndices)
+		for (int32 LODIndex = 0; LODIndex < MAX_STATIC_MESH_LODS; ++LODIndex)
 		{
-			//Assign the new material index to the raw mesh
-			FaceMaterialIndex = Map[FaceMaterialIndex];
+			FRawMesh& RawMesh = SourceMeshes[MeshIndex].MeshLOD[LODIndex];
+			const auto* Map = MaterialMap.Find(FMeshIdAndLOD(MeshIndex, LODIndex));
+			if (Map)
+			{
+				for (int32& FaceMaterialIndex : RawMesh.FaceMaterialIndices)
+				{
+					//Assign the new material index to the raw mesh
+					FaceMaterialIndex = (*Map)[FaceMaterialIndex];
+				}
+			}
 		}
 	}
 
@@ -4111,68 +4243,76 @@ void FMeshUtilities::MergeActors(
 	// Merge meshes into single mesh
 	for (int32 SourceMeshIdx = 0; SourceMeshIdx < SourceMeshes.Num(); ++SourceMeshIdx)
 	{
-		// Merge vertex data from source mesh list into single mesh
-		FRawMesh& TargetRawMesh = MergedMesh.Mesh;
-		const FRawMesh& SourceRawMesh = SourceMeshes[SourceMeshIdx].Mesh;
-						
-		TargetRawMesh.FaceSmoothingMasks.Append(SourceRawMesh.FaceSmoothingMasks);
-		TargetRawMesh.FaceMaterialIndices.Append(SourceRawMesh.FaceMaterialIndices);
-
-		int32 IndicesOffset = TargetRawMesh.VertexPositions.Num();
-
-		for (int32 Index : SourceRawMesh.WedgeIndices)
+		for (int32 LODIndex = 0; LODIndex < MAX_STATIC_MESH_LODS; ++LODIndex)
 		{
-			TargetRawMesh.WedgeIndices.Add(Index + IndicesOffset);
-		}
+			// Merge vertex data from source mesh list into single mesh
+			FRawMesh& TargetRawMesh = MergedMesh.MeshLOD[LODIndex];
+			const FRawMesh& SourceRawMesh = SourceMeshes[SourceMeshIdx].MeshLOD[LODIndex];
+			
+			if (SourceRawMesh.VertexPositions.Num() == 0)
+			{
+				continue;
+			}
+									
+			TargetRawMesh.FaceSmoothingMasks.Append(SourceRawMesh.FaceSmoothingMasks);
+			TargetRawMesh.FaceMaterialIndices.Append(SourceRawMesh.FaceMaterialIndices);
 
-		for (FVector VertexPos : SourceRawMesh.VertexPositions)
-		{
-			TargetRawMesh.VertexPositions.Add(VertexPos - MergedMesh.Pivot);
-		}
-						
-		TargetRawMesh.WedgeTangentX.Append(SourceRawMesh.WedgeTangentX);
-		TargetRawMesh.WedgeTangentY.Append(SourceRawMesh.WedgeTangentY);
-		TargetRawMesh.WedgeTangentZ.Append(SourceRawMesh.WedgeTangentZ);
-		
-		// Deal with vertex colors
-		// Some meshes may have it, in this case merged mesh will be forced to have vertex colors as well
-		if (bWithVertexColors)
-		{
-			if (SourceRawMesh.WedgeColors.Num())
+			int32 IndicesOffset = TargetRawMesh.VertexPositions.Num();
+
+			for (int32 Index : SourceRawMesh.WedgeIndices)
 			{
-				TargetRawMesh.WedgeColors.Append(SourceRawMesh.WedgeColors);
+				TargetRawMesh.WedgeIndices.Add(Index + IndicesOffset);
 			}
-			else
+
+			for (FVector VertexPos : SourceRawMesh.VertexPositions)
 			{
-				// In case this source mesh does not have vertex colors, fill target with 0xFF
-				int32 ColorsOffset = TargetRawMesh.WedgeColors.Num();
-				int32 ColorsNum = SourceRawMesh.WedgeIndices.Num();
-				TargetRawMesh.WedgeColors.AddUninitialized(ColorsNum);
-				FMemory::Memset(&TargetRawMesh.WedgeColors[ColorsOffset], 0xFF, ColorsNum*TargetRawMesh.WedgeColors.GetTypeSize());
+				TargetRawMesh.VertexPositions.Add(VertexPos - MergedMesh.Pivot);
 			}
-		}
+						
+			TargetRawMesh.WedgeTangentX.Append(SourceRawMesh.WedgeTangentX);
+			TargetRawMesh.WedgeTangentY.Append(SourceRawMesh.WedgeTangentY);
+			TargetRawMesh.WedgeTangentZ.Append(SourceRawMesh.WedgeTangentZ);
 		
-		// Merge all other UV channels 
-		for (int32 ChannelIdx = 0; ChannelIdx < MAX_MESH_TEXTURE_COORDS; ++ChannelIdx)
-		{
-			// Whether this channel has data
-			if (bOcuppiedUVChannels[ChannelIdx])
+			// Deal with vertex colors
+			// Some meshes may have it, in this case merged mesh will be forced to have vertex colors as well
+			if (bWithVertexColors[LODIndex])
 			{
-				const TArray<FVector2D>& SourceChannel = SourceRawMesh.WedgeTexCoords[ChannelIdx];
-				TArray<FVector2D>& TargetChannel = TargetRawMesh.WedgeTexCoords[ChannelIdx];
-
-				// Whether source mesh has data in this channel
-				if (SourceChannel.Num())
+				if (SourceRawMesh.WedgeColors.Num())
 				{
-					TargetChannel.Append(SourceChannel);
+					TargetRawMesh.WedgeColors.Append(SourceRawMesh.WedgeColors);
 				}
 				else
 				{
-					// Fill with zero coordinates if source mesh has no data for this channel
-					const int32 TexCoordNum = SourceRawMesh.WedgeIndices.Num();
-					for (int32 CoordIdx = 0; CoordIdx < TexCoordNum; ++CoordIdx)
+					// In case this source mesh does not have vertex colors, fill target with 0xFF
+					int32 ColorsOffset = TargetRawMesh.WedgeColors.Num();
+					int32 ColorsNum = SourceRawMesh.WedgeIndices.Num();
+					TargetRawMesh.WedgeColors.AddUninitialized(ColorsNum);
+					FMemory::Memset(&TargetRawMesh.WedgeColors[ColorsOffset], 0xFF, ColorsNum*TargetRawMesh.WedgeColors.GetTypeSize());
+				}
+			}
+		
+			// Merge all other UV channels 
+			for (int32 ChannelIdx = 0; ChannelIdx < MAX_MESH_TEXTURE_COORDS; ++ChannelIdx)
+			{
+				// Whether this channel has data
+				if (bOcuppiedUVChannels[LODIndex][ChannelIdx])
+				{
+					const TArray<FVector2D>& SourceChannel = SourceRawMesh.WedgeTexCoords[ChannelIdx];
+					TArray<FVector2D>& TargetChannel = TargetRawMesh.WedgeTexCoords[ChannelIdx];
+
+					// Whether source mesh has data in this channel
+					if (SourceChannel.Num())
 					{
-						TargetChannel.Add(FVector2D::ZeroVector);
+						TargetChannel.Append(SourceChannel);
+					}
+					else
+					{
+						// Fill with zero coordinates if source mesh has no data for this channel
+						const int32 TexCoordNum = SourceRawMesh.WedgeIndices.Num();
+						for (int32 CoordIdx = 0; CoordIdx < TexCoordNum; ++CoordIdx)
+						{
+							TargetChannel.Add(FVector2D::ZeroVector);
+						}
 					}
 				}
 			}
@@ -4208,20 +4348,26 @@ void FMeshUtilities::MergeActors(
 
 		// make sure it has a new lighting guid
 		StaticMesh->LightingGuid = FGuid::NewGuid();
-
-
-		FStaticMeshSourceModel* SrcModel = new (StaticMesh->SourceModels) FStaticMeshSourceModel();
-		/*Don't allow the engine to recalculate normals*/
-		SrcModel->BuildSettings.bRecomputeNormals = false;
-		SrcModel->BuildSettings.bRecomputeTangents = false;
-		SrcModel->BuildSettings.bRemoveDegenerates = false;
-		SrcModel->BuildSettings.bUseFullPrecisionUVs = false;
-		SrcModel->BuildSettings.bGenerateLightmapUVs = InSettings.bGenerateLightMapUV;
-		SrcModel->BuildSettings.MinLightmapResolution = InSettings.TargetLightMapResolution;
-		SrcModel->BuildSettings.SrcLightmapIndex = 0;
-		SrcModel->BuildSettings.DstLightmapIndex = InSettings.TargetLightMapUVChannel;
-
-		SrcModel->RawMeshBulkData->SaveRawMesh(MergedMesh.Mesh);
+		
+		for (int32 LODIndex = 0; LODIndex < NumMaxLOD; ++LODIndex)
+		{
+			FRawMesh& MergedMeshLOD = MergedMesh.MeshLOD[LODIndex];
+			if (MergedMeshLOD.VertexPositions.Num() > 0)
+			{
+				FStaticMeshSourceModel* SrcModel = new (StaticMesh->SourceModels) FStaticMeshSourceModel();
+				/*Don't allow the engine to recalculate normals*/
+				SrcModel->BuildSettings.bRecomputeNormals = false;
+				SrcModel->BuildSettings.bRecomputeTangents = false;
+				SrcModel->BuildSettings.bRemoveDegenerates = false;
+				SrcModel->BuildSettings.bUseFullPrecisionUVs = false;
+				SrcModel->BuildSettings.bGenerateLightmapUVs = InSettings.bGenerateLightMapUV;
+				SrcModel->BuildSettings.MinLightmapResolution = InSettings.TargetLightMapResolution;
+				SrcModel->BuildSettings.SrcLightmapIndex = 0;
+				SrcModel->BuildSettings.DstLightmapIndex = InSettings.TargetLightMapUVChannel;
+
+				SrcModel->RawMeshBulkData->SaveRawMesh(MergedMeshLOD);
+			}
+		}
 
 		// Assign materials
 		for (UMaterialInterface* Material : UniqueMaterials)
