Fixed typo in static mesh to rawmesh exporting

diff --git a/Engine/Source/Developer/MeshUtilities/Private/MeshUtilities.cpp b/Engine/Source/Developer/MeshUtilities/Private/MeshUtilities.cpp
--- a/Engine/Source/Developer/MeshUtilities/Private/MeshUtilities.cpp
+++ b/Engine/Source/Developer/MeshUtilities/Private/MeshUtilities.cpp
@@ -4017,7 +4017,7 @@ bool FMeshUtilities::ConstructRawMesh(
 	}
 
 	// Compute any missing tangents.
-	if (bRecomputeTangents)
+	if (bRecomputeNormals || bRecomputeTangents)
 	{
 		float ComparisonThreshold = GetComparisonThreshold(BuildSettings);
 		TMultiMap<int32,int32> OverlappingCorners;
