#ue4 - Change use of FTransform::TransformFVector4() to FTransform::TransformPosition() where possible. This is faster.

diff --git a/Engine/Source/Developer/MeshUtilities/Private/MeshUtilities.cpp b/Engine/Source/Developer/MeshUtilities/Private/MeshUtilities.cpp
--- a/Engine/Source/Developer/MeshUtilities/Private/MeshUtilities.cpp
+++ b/Engine/Source/Developer/MeshUtilities/Private/MeshUtilities.cpp
@@ -3964,7 +3964,7 @@ bool FMeshUtilities::ConstructRawMesh(
 	FTransform CtoM = InMeshComponent->ComponentToWorld;
 	for (FVector& Vertex : OutRawMesh.VertexPositions)
 	{
-		Vertex = CtoM.TransformFVector4(Vertex);
+		Vertex = CtoM.TransformPosition(Vertex);
 	}
 	for (FVector& TangentX : OutRawMesh.WedgeTangentX)
 	{
