Added support for material merging to actor merging utility (supports only opaque materials)

diff --git a/Engine/Source/Developer/MeshUtilities/Private/MeshUtilities.cpp b/Engine/Source/Developer/MeshUtilities/Private/MeshUtilities.cpp
--- a/Engine/Source/Developer/MeshUtilities/Private/MeshUtilities.cpp
+++ b/Engine/Source/Developer/MeshUtilities/Private/MeshUtilities.cpp
@@ -4128,6 +4128,165 @@ bool PropagatePaintedColorsToRawMesh(UStaticMeshComponent* StaticMeshComponent,
 	return false;
 }
 
+static void CopyTextureRect(const FColor* Src, const FIntPoint& SrcSize, FColor* Dst, const FIntPoint& DstSize, const FIntPoint& DstPos)
+{
+	int32 RowLength = SrcSize.X*sizeof(FColor);
+	FColor* RowDst = Dst + DstSize.X*DstPos.Y;
+	const FColor* RowSrc = Src;
+	
+	for (int32 RowIdx = 0; RowIdx < SrcSize.Y; ++RowIdx)
+	{
+		FMemory::Memcpy(RowDst + DstPos.X, RowSrc, RowLength);
+		
+		RowDst+= DstSize.X;
+		RowSrc+= SrcSize.X;
+	}
+}
+
+struct FRawMeshExt
+{
+	FRawMeshExt() 
+	{}
+		
+	FRawMesh	MeshLOD[MAX_STATIC_MESH_LODS];
+	FString		AssetPackageName;
+	FVector		Pivot;	
+};
+
+static void MergeMaterials(UWorld* InWorld, const TArray<UMaterialInterface*>& InMaterialList, TArray<FRawMeshExt>& InOutMeshes, MaterialExportUtils::FFlattenMaterial& OutMergedMaterial)
+{
+	using namespace MaterialExportUtils;
+	
+	int32 AtlasGridSize = FMath::CeilToInt(FMath::Sqrt(InMaterialList.Num()));
+	FIntPoint AtlasTextureSize = OutMergedMaterial.DiffuseSize;
+	FIntPoint ExportTextureSize	= AtlasTextureSize/AtlasGridSize;
+	int32 AtlasNumSamples = AtlasTextureSize.X*AtlasTextureSize.Y;
+	
+	bool bExportNormal		= (OutMergedMaterial.NormalSize != FIntPoint::ZeroValue);
+	bool bExportMetallic	= (OutMergedMaterial.MetallicSize != FIntPoint::ZeroValue);
+	bool bExportRoughness	= (OutMergedMaterial.RoughnessSize != FIntPoint::ZeroValue);
+	bool bExportSpecular	= (OutMergedMaterial.SpecularSize != FIntPoint::ZeroValue);
+
+	// Pre-allocate buffers for texture atlas
+	OutMergedMaterial.DiffuseSamples.Reserve(AtlasNumSamples);
+	OutMergedMaterial.DiffuseSamples.SetNumZeroed(AtlasNumSamples);
+	if (bExportNormal)
+	{
+		check(OutMergedMaterial.NormalSize == OutMergedMaterial.DiffuseSize);
+		OutMergedMaterial.NormalSamples.Reserve(AtlasNumSamples);
+		OutMergedMaterial.NormalSamples.SetNumZeroed(AtlasNumSamples);
+	}
+	if (bExportMetallic)
+	{
+		check(OutMergedMaterial.MetallicSize == OutMergedMaterial.DiffuseSize);
+		OutMergedMaterial.MetallicSamples.Reserve(AtlasNumSamples);
+		OutMergedMaterial.MetallicSamples.SetNumZeroed(AtlasNumSamples);
+	}
+	if (bExportRoughness)
+	{
+		check(OutMergedMaterial.RoughnessSize == OutMergedMaterial.DiffuseSize);
+		OutMergedMaterial.RoughnessSamples.Reserve(AtlasNumSamples);
+		OutMergedMaterial.RoughnessSamples.SetNumZeroed(AtlasNumSamples);
+	}
+	if (bExportSpecular)
+	{
+		check(OutMergedMaterial.SpecularSize == OutMergedMaterial.DiffuseSize);
+		OutMergedMaterial.SpecularSamples.Reserve(AtlasNumSamples);
+		OutMergedMaterial.SpecularSamples.SetNumZeroed(AtlasNumSamples);
+	}
+
+	int32 AtlasRowIdx = 0;
+	int32 AtlasColIdx = 0;
+	FIntPoint AtlasTargetPos = FIntPoint(0,0);
+	TArray<FVector2D> UVOffsets;
+					
+	// Flatten all materials and merge them into one material using texture atlases
+	for (const UMaterialInterface* Material : InMaterialList)
+	{
+		FFlattenMaterial FlatMaterial;
+		FlatMaterial.DiffuseSize	= ExportTextureSize;
+		FlatMaterial.NormalSize		= bExportNormal ? ExportTextureSize : FIntPoint::ZeroValue;
+		FlatMaterial.MetallicSize	= bExportMetallic ? ExportTextureSize : FIntPoint::ZeroValue;
+		FlatMaterial.RoughnessSize	= bExportRoughness ? ExportTextureSize : FIntPoint::ZeroValue;;
+		FlatMaterial.SpecularSize	= bExportSpecular ? ExportTextureSize : FIntPoint::ZeroValue;;
+
+		ExportMaterial(InWorld, const_cast<UMaterialInterface*>(Material), FlatMaterial);
+
+		if (FlatMaterial.DiffuseSamples.Num() > 0)
+		{
+			CopyTextureRect(FlatMaterial.DiffuseSamples.GetData(), ExportTextureSize, OutMergedMaterial.DiffuseSamples.GetData(), AtlasTextureSize, AtlasTargetPos);
+		}
+
+		if (FlatMaterial.NormalSamples.Num() > 0)
+		{
+			CopyTextureRect(FlatMaterial.NormalSamples.GetData(), ExportTextureSize, OutMergedMaterial.NormalSamples.GetData(), AtlasTextureSize, AtlasTargetPos);
+		}
+
+		if (FlatMaterial.MetallicSamples.Num() > 0)
+		{
+			CopyTextureRect(FlatMaterial.MetallicSamples.GetData(), ExportTextureSize, OutMergedMaterial.MetallicSamples.GetData(), AtlasTextureSize, AtlasTargetPos);
+		}
+
+		if (FlatMaterial.RoughnessSamples.Num() > 0)
+		{
+			CopyTextureRect(FlatMaterial.RoughnessSamples.GetData(), ExportTextureSize, OutMergedMaterial.RoughnessSamples.GetData(), AtlasTextureSize, AtlasTargetPos);
+		}
+
+		if (FlatMaterial.SpecularSamples.Num() > 0)
+		{
+			CopyTextureRect(FlatMaterial.SpecularSamples.GetData(), ExportTextureSize, OutMergedMaterial.SpecularSamples.GetData(), AtlasTextureSize, AtlasTargetPos);
+		}
+
+		UVOffsets.Add(FVector2D(
+			(float)AtlasTargetPos.X/AtlasTextureSize.X, 
+			(float)AtlasTargetPos.Y/AtlasTextureSize.Y));
+		
+		AtlasColIdx++;
+		if (AtlasColIdx >= AtlasGridSize)
+		{
+			AtlasColIdx = 0;
+			AtlasRowIdx++;
+		}
+
+		AtlasTargetPos = FIntPoint(AtlasColIdx*ExportTextureSize.X, AtlasRowIdx*ExportTextureSize.Y);
+	}
+
+	// Adjust UVs
+	FVector2D UVScale = FVector2D(
+		(float)ExportTextureSize.X/AtlasTextureSize.X, 
+		(float)ExportTextureSize.Y/AtlasTextureSize.Y);
+
+	for (int32 MeshIndex = 0; MeshIndex < InOutMeshes.Num(); ++MeshIndex)
+	{
+		for (int32 LODIndex = 0; LODIndex < MAX_STATIC_MESH_LODS; ++LODIndex)
+		{
+			FRawMesh& RawMesh = InOutMeshes[MeshIndex].MeshLOD[LODIndex];
+
+			for (int32 UVChannelIdx = 0; UVChannelIdx < MAX_MESH_TEXTURE_COORDS; ++UVChannelIdx)
+			{
+				TArray<FVector2D>& UVs = RawMesh.WedgeTexCoords[UVChannelIdx];
+				if (UVs.Num() > 0)
+				{
+					for (int32 i = 0; i < UVs.Num(); i+=3)
+					{
+						int32 FaceMaterialIndex = RawMesh.FaceMaterialIndices[i/3];
+						FVector2D UVOffset = UVOffsets[FaceMaterialIndex];
+
+						UVs[i+0] = UVs[i+0]*UVScale + UVOffset;
+						UVs[i+1] = UVs[i+1]*UVScale + UVOffset;
+						UVs[i+2] = UVs[i+2]*UVScale + UVOffset;
+					}
+				}
+			}
+			
+			for (int32& FaceMaterialIndex : RawMesh.FaceMaterialIndices)
+			{
+				FaceMaterialIndex = 0;
+			}
+		}
+	}
+}
+
 void FMeshUtilities::MergeActors(
 	const TArray<AActor*>& SourceActors, 
 	const FMeshMergingSettings& InSettings, 
@@ -4154,16 +4313,6 @@ void FMeshUtilities::MergeActors(
 			}
 		}
 	}
-
-	struct FRawMeshExt
-	{
-		FRawMeshExt() 
-		{}
-		
-		FRawMesh	MeshLOD[MAX_STATIC_MESH_LODS];
-		FString		AssetPackageName;
-		FVector		Pivot;	
-	};
 	
 	typedef FIntPoint FMeshIdAndLOD;
 	
@@ -4172,8 +4321,7 @@ void FMeshUtilities::MergeActors(
 	TArray<FRawMeshExt>								SourceMeshes;
 	bool											bWithVertexColors[MAX_STATIC_MESH_LODS] = {};
 	bool											bOcuppiedUVChannels[MAX_STATIC_MESH_LODS][MAX_MESH_TEXTURE_COORDS] = {};
-	
-	
+		
 	// Convert collected static mesh components into raw meshes
 	SourceMeshes.Reserve(ComponentsToMerge.Num());
 	SourceMeshes.SetNum(ComponentsToMerge.Num());
@@ -4249,6 +4397,24 @@ void FMeshUtilities::MergeActors(
 		}
 	}
 
+	MaterialExportUtils::FFlattenMaterial MergedFlatMaterial;
+	
+	if (InSettings.bMergeMaterials)
+	{
+		FIntPoint AtlasTextureSize			= FIntPoint(InSettings.MergedMaterialAtlasResolution, InSettings.MergedMaterialAtlasResolution);
+		MergedFlatMaterial.DiffuseSize		= AtlasTextureSize;
+		MergedFlatMaterial.NormalSize		= InSettings.bExportNormalMap ? AtlasTextureSize : FIntPoint::ZeroValue;
+		MergedFlatMaterial.MetallicSize		= InSettings.bExportMetallicMap ? AtlasTextureSize : FIntPoint::ZeroValue;
+		MergedFlatMaterial.RoughnessSize	= InSettings.bExportRoughnessMap ? AtlasTextureSize : FIntPoint::ZeroValue;
+		MergedFlatMaterial.SpecularSize		= InSettings.bExportSpecularMap ? AtlasTextureSize : FIntPoint::ZeroValue;
+			
+		UWorld* World = SourceActors[0]->GetWorld();
+		MergeMaterials(World, UniqueMaterials, SourceMeshes, MergedFlatMaterial);
+		
+		// We assume all materials was merged into MergedFlatMaterial
+		UniqueMaterials.Empty();
+	}
+
 	FRawMeshExt MergedMesh;
 
 	// Use first mesh for naming and pivot
@@ -4384,6 +4550,30 @@ void FMeshUtilities::MergeActors(
 			}
 		}
 
+		if (InSettings.bMergeMaterials)
+		{
+			FString MaterialAssetName;
+			FString MaterialPackageName;
+			if (InPackageName.IsEmpty())
+			{
+				MaterialAssetName = TEXT("M_MERGED_") + FPackageName::GetShortName(MergedMesh.AssetPackageName);
+				MaterialPackageName = FPackageName::GetLongPackagePath(MergedMesh.AssetPackageName) + TEXT("/") + MaterialAssetName;
+			}
+			else
+			{
+				MaterialAssetName = FPackageName::GetShortName(InPackageName);
+				MaterialPackageName = InPackageName;
+			}
+
+			UPackage* MaterialPackage = CreatePackage(NULL, *MaterialPackageName);
+			check(MaterialPackage);
+			MaterialPackage->FullyLoad();
+			MaterialPackage->Modify();
+
+			UMaterial* MergedMaterial = CreateMaterial(MergedFlatMaterial, MaterialPackage, MaterialAssetName, RF_Public|RF_Standalone, OutAssetsToSync);
+			UniqueMaterials.Add(MergedMaterial);
+		}
+	
 		// Assign materials
 		for (UMaterialInterface* Material : UniqueMaterials)
 		{
diff --git a/Engine/Source/Developer/MeshUtilities/Public/MeshUtilities.h b/Engine/Source/Developer/MeshUtilities/Public/MeshUtilities.h
--- a/Engine/Source/Developer/MeshUtilities/Public/MeshUtilities.h
+++ b/Engine/Source/Developer/MeshUtilities/Public/MeshUtilities.h
@@ -87,6 +87,19 @@ struct FMeshMergingSettings
 	
 	/** Whether merged mesh should have pivot at world origin, or at first merged component otherwise */
 	bool bPivotPointAtZero;
+
+	/** Whether to merge source materials into one flat material */
+	bool bMergeMaterials;
+	/** Whether to export normal maps for material merging */
+	bool bExportNormalMap;
+	/** Whether to export metallic maps for material merging */
+	bool bExportMetallicMap;
+	/** Whether to export roughness maps for material merging */
+	bool bExportRoughnessMap;
+	/** Whether to export specular maps for material merging */
+	bool bExportSpecularMap;
+	/** Merged material texture atlas resolution */
+	int32 MergedMaterialAtlasResolution;
 		
 	/** Default settings. */
 	FMeshMergingSettings()
@@ -95,6 +108,12 @@ struct FMeshMergingSettings
 		, TargetLightMapResolution(256)
 		, bImportVertexColors(false)
 		, bPivotPointAtZero(false)
+		, bMergeMaterials(false)
+		, bExportNormalMap(true)
+		, bExportMetallicMap(false)
+		, bExportRoughnessMap(false)
+		, bExportSpecularMap(false)
+		, MergedMaterialAtlasResolution(1024)
 	{
 	}
 };
diff --git a/Engine/Source/Editor/UnrealEd/Private/Dialogs/SMeshMergingDialog.cpp b/Engine/Source/Editor/UnrealEd/Private/Dialogs/SMeshMergingDialog.cpp
--- a/Engine/Source/Editor/UnrealEd/Private/Dialogs/SMeshMergingDialog.cpp
+++ b/Engine/Source/Editor/UnrealEd/Private/Dialogs/SMeshMergingDialog.cpp
@@ -32,6 +32,7 @@ SMeshMergingDialog::~SMeshMergingDialog()
 	LevelEditor.OnActorSelectionChanged().RemoveAll(this);
 }
 
+BEGIN_SLATE_FUNCTION_BUILD_OPTIMIZATION
 void SMeshMergingDialog::Construct(const FArguments& InArgs)
 {
 	ParentWindow = InArgs._ParentWindow;
@@ -40,15 +41,17 @@ void SMeshMergingDialog::Construct(const FArguments& InArgs)
 	LevelEditor.OnActorSelectionChanged().AddSP(this, &SMeshMergingDialog::OnActorSelectionChanged);
 	
 
-	// Setup available resolutions for lightmap
+	// Setup available resolutions for lightmap and merged materials
 	const int32 MinTexResolution = 1 << FTextureLODGroup().MinLODMipCount;
 	const int32 MaxTexResolution = 1 << FTextureLODGroup().MaxLODMipCount;
-	for (int32 LightmapRes = MinTexResolution; LightmapRes <= MaxTexResolution; LightmapRes*=2)
+	for (int32 TexRes = MinTexResolution; TexRes <= MaxTexResolution; TexRes*=2)
 	{
-		LightMapResolutionOptions.Add(MakeShareable(new FString(FString::FormatAsNumber(LightmapRes))));
+		LightMapResolutionOptions.Add(MakeShareable(new FString(FString::FormatAsNumber(TexRes))));
+		MergedMaterialResolutionOptions.Add(MakeShareable(new FString(FString::FormatAsNumber(TexRes))));
 	}
 
 	MergingSettings.TargetLightMapResolution = FMath::Clamp(MergingSettings.TargetLightMapResolution, MinTexResolution, MaxTexResolution);
+	MergingSettings.MergedMaterialAtlasResolution = FMath::Clamp(MergingSettings.MergedMaterialAtlasResolution, MinTexResolution, MaxTexResolution);
 		
 	// Setup available UV channels for an atlased lightmap
 	for (int32 Index = 0; Index < MAX_MESH_TEXTURE_COORDS; Index++)
@@ -192,7 +195,109 @@ void SMeshMergingDialog::Construct(const FArguments& InArgs)
 						SNew(STextBlock).Text(LOCTEXT("PlaceInWorldLabel", "Place In World"))
 					]
 				]
-						
+
+				// Merge materials
+				+SVerticalBox::Slot()
+				.AutoHeight()
+				.Padding(FEditorStyle::GetMargin("StandardDialog.ContentPadding"))
+				[
+					SNew(SCheckBox)
+					.Type(ESlateCheckBoxType::CheckBox)
+					.IsChecked(this, &SMeshMergingDialog::GetMergeMaterials)
+					.OnCheckStateChanged(this, &SMeshMergingDialog::SetMergeMaterials)
+					.Content()
+					[
+						SNew(STextBlock).Text(LOCTEXT("MergeMaterialsLabel", "Merge Materials"))
+					]
+				]
+				// Export normal map
+				+SVerticalBox::Slot()
+				.AutoHeight()
+				.Padding(FEditorStyle::GetMargin("StandardDialog.ContentPadding"))
+				[
+					SNew(SCheckBox)
+					.Type(ESlateCheckBoxType::CheckBox)
+					.IsChecked(this, &SMeshMergingDialog::GetExportNormalMap)
+					.OnCheckStateChanged(this, &SMeshMergingDialog::SetExportNormalMap)
+					.IsEnabled(this, &SMeshMergingDialog::IsMaterialMergingEnabled)
+					.Content()
+					[
+						SNew(STextBlock).Text(LOCTEXT("ExportNormalMapLabel", "Export Normal Map"))
+					]
+				]
+				// Export metallic map
+				+SVerticalBox::Slot()
+				.AutoHeight()
+				.Padding(FEditorStyle::GetMargin("StandardDialog.ContentPadding"))
+				[
+					SNew(SCheckBox)
+					.Type(ESlateCheckBoxType::CheckBox)
+					.IsChecked(this, &SMeshMergingDialog::GetExportMetallicMap)
+					.OnCheckStateChanged(this, &SMeshMergingDialog::SetExportMetallicMap)
+					.IsEnabled(this, &SMeshMergingDialog::IsMaterialMergingEnabled)
+					.Content()
+					[
+						SNew(STextBlock).Text(LOCTEXT("ExportMetallicMapLabel", "Export Metallic Map"))
+					]
+				]
+				// Export roughness map
+				+SVerticalBox::Slot()
+				.AutoHeight()
+				.Padding(FEditorStyle::GetMargin("StandardDialog.ContentPadding"))
+				[
+					SNew(SCheckBox)
+					.Type(ESlateCheckBoxType::CheckBox)
+					.IsChecked(this, &SMeshMergingDialog::GetExportRoughnessMap)
+					.OnCheckStateChanged(this, &SMeshMergingDialog::SetExportRoughnessMap)
+					.IsEnabled(this, &SMeshMergingDialog::IsMaterialMergingEnabled)
+					.Content()
+					[
+						SNew(STextBlock).Text(LOCTEXT("ExportRoughnessMapLabel", "Export Roughness Map"))
+					]
+				]
+				// Export specular map
+				+SVerticalBox::Slot()
+				.AutoHeight()
+				.Padding(FEditorStyle::GetMargin("StandardDialog.ContentPadding"))
+				[
+					SNew(SCheckBox)
+					.Type(ESlateCheckBoxType::CheckBox)
+					.IsChecked(this, &SMeshMergingDialog::GetExportSpecularMap)
+					.OnCheckStateChanged(this, &SMeshMergingDialog::SetExportSpecularMap)
+					.IsEnabled(this, &SMeshMergingDialog::IsMaterialMergingEnabled)
+					.Content()
+					[
+						SNew(STextBlock).Text(LOCTEXT("ExportSpecularMapLabel", "Export Specular Map"))
+					]
+				]
+				
+				// Merged texture size
+				+SVerticalBox::Slot()
+				.AutoHeight()
+				.Padding(FEditorStyle::GetMargin("StandardDialog.ContentPadding"))
+				[
+					SNew(SHorizontalBox)
+					+SHorizontalBox::Slot()
+					.AutoWidth()
+					.VAlign(VAlign_Center)
+					[
+						SNew(STextBlock)
+						.Text(LOCTEXT("MergedMaterialAtlasResolutionLabel", "Merged Material Atlas Resolution:"))
+					]
+												
+					+SHorizontalBox::Slot()
+					.AutoWidth()
+					.VAlign(VAlign_Center)
+					.Padding(4,0,4,0)
+					[
+						SNew(STextComboBox)
+						.IsEnabled(this, &SMeshMergingDialog::IsMaterialMergingEnabled)
+						.OptionsSource(&MergedMaterialResolutionOptions)
+						.InitiallySelectedItem(MergedMaterialResolutionOptions[FMath::FloorLog2(MergingSettings.MergedMaterialAtlasResolution)])
+						.OnSelectionChanged(this, &SMeshMergingDialog::SetMergedMaterialAtlasResolution)
+					]
+				]
+														
 				// Asset name and picker
 				+SVerticalBox::Slot()
 				.AutoHeight()
@@ -262,6 +367,7 @@ void SMeshMergingDialog::Construct(const FArguments& InArgs)
 		]
 	];
 }
+END_SLATE_FUNCTION_BUILD_OPTIMIZATION
 
 FReply SMeshMergingDialog::OnCancelClicked()
 {
@@ -345,6 +451,66 @@ void SMeshMergingDialog::SetPlaceInWorld(ECheckBoxState NewValue)
 	bPlaceInWorld = (ECheckBoxState::Checked == NewValue);
 }
 
+bool SMeshMergingDialog::IsMaterialMergingEnabled() const
+{
+	return MergingSettings.bMergeMaterials;
+}
+
+ECheckBoxState SMeshMergingDialog::GetMergeMaterials() const
+{
+	return (MergingSettings.bMergeMaterials ? ECheckBoxState::Checked : ECheckBoxState::Unchecked);
+}
+
+void SMeshMergingDialog::SetMergeMaterials(ECheckBoxState NewValue)
+{
+	MergingSettings.bMergeMaterials = (ECheckBoxState::Checked == NewValue);
+}
+
+ECheckBoxState SMeshMergingDialog::GetExportNormalMap() const
+{
+	return (MergingSettings.bExportNormalMap ? ECheckBoxState::Checked : ECheckBoxState::Unchecked);
+}
+
+void SMeshMergingDialog::SetExportNormalMap(ECheckBoxState NewValue)
+{
+	MergingSettings.bExportNormalMap = (ECheckBoxState::Checked == NewValue);
+}
+
+ECheckBoxState SMeshMergingDialog::GetExportMetallicMap() const
+{
+	return (MergingSettings.bExportMetallicMap ? ECheckBoxState::Checked : ECheckBoxState::Unchecked);
+}
+
+void SMeshMergingDialog::SetExportMetallicMap(ECheckBoxState NewValue)
+{
+	MergingSettings.bExportMetallicMap = (ECheckBoxState::Checked == NewValue);
+}
+
+ECheckBoxState SMeshMergingDialog::GetExportRoughnessMap() const
+{
+	return (MergingSettings.bExportRoughnessMap ? ECheckBoxState::Checked : ECheckBoxState::Unchecked);
+}
+
+void SMeshMergingDialog::SetExportRoughnessMap(ECheckBoxState NewValue)
+{
+	MergingSettings.bExportRoughnessMap = (ECheckBoxState::Checked == NewValue);
+}
+
+ECheckBoxState SMeshMergingDialog::GetExportSpecularMap() const
+{
+	return (MergingSettings.bExportSpecularMap ? ECheckBoxState::Checked : ECheckBoxState::Unchecked);
+}
+
+void SMeshMergingDialog::SetExportSpecularMap(ECheckBoxState NewValue)
+{
+	MergingSettings.bExportSpecularMap = (ECheckBoxState::Checked == NewValue);
+}
+
+void SMeshMergingDialog::SetMergedMaterialAtlasResolution(TSharedPtr<FString> NewSelection, ESelectInfo::Type SelectInfo)
+{
+	TTypeFromString<int32>::FromString(MergingSettings.MergedMaterialAtlasResolution, **NewSelection);
+}
+
 void SMeshMergingDialog::GenerateNewPackageName()
 {
 	MergedMeshPackageName = FPackageName::FilenameToLongPackageName(FPaths::GameContentDir() + TEXT("SM_MERGED"));
@@ -417,7 +583,7 @@ void SMeshMergingDialog::RunMerging()
 		OpenMsgDlgInt(EAppMsgType::Ok, Message, NSLOCTEXT("UnrealEd", "FailedToMergeActors_Title", "Unable to merge actors"));
 		return;
 	}
-	
+		
 	FVector MergedActorLocation;
 	TArray<UObject*> AssetsToSync;
 	MeshUtilities.MergeActors(Actors, MergingSettings, MergedMeshPackageName, AssetsToSync, MergedActorLocation);
diff --git a/Engine/Source/Editor/UnrealEd/Private/EditorGroup.cpp b/Engine/Source/Editor/UnrealEd/Private/EditorGroup.cpp
--- a/Engine/Source/Editor/UnrealEd/Private/EditorGroup.cpp
+++ b/Engine/Source/Editor/UnrealEd/Private/EditorGroup.cpp
@@ -266,7 +266,7 @@ void UUnrealEdEngine::edactMergeActorsByMaterials()
 	/** Create the window to host mesh megin widget */
 	TSharedRef<SWindow> MergeActorsWindow = SNew(SWindow)
 											.Title(FText::FromString("Merge actors"))
-											.ClientSize(FVector2D(400, 200));
+											.ClientSize(FVector2D(400, 330));
 
 	/** Set the content of the window to our package dialog widget */
 	TSharedRef<SMeshMergingDialog> MergingDialog = SNew(SMeshMergingDialog)
diff --git a/Engine/Source/Editor/UnrealEd/Public/Dialogs/SMeshMergingDialog.h b/Engine/Source/Editor/UnrealEd/Public/Dialogs/SMeshMergingDialog.h
--- a/Engine/Source/Editor/UnrealEd/Public/Dialogs/SMeshMergingDialog.h
+++ b/Engine/Source/Editor/UnrealEd/Public/Dialogs/SMeshMergingDialog.h
@@ -52,6 +52,25 @@ private:
 
 	ECheckBoxState GetPlaceInWorld() const;
 	void SetPlaceInWorld(ECheckBoxState NewValue);
+
+	/** Material merging */
+	bool IsMaterialMergingEnabled() const;
+	ECheckBoxState GetMergeMaterials() const;
+	void SetMergeMaterials(ECheckBoxState NewValue);
+
+	ECheckBoxState GetExportNormalMap() const;
+	void SetExportNormalMap(ECheckBoxState NewValue);
+
+	ECheckBoxState GetExportMetallicMap() const;
+	void SetExportMetallicMap(ECheckBoxState NewValue);
+
+	ECheckBoxState GetExportRoughnessMap() const;
+	void SetExportRoughnessMap(ECheckBoxState NewValue);
+
+	ECheckBoxState GetExportSpecularMap() const;
+	void SetExportSpecularMap(ECheckBoxState NewValue);
+
+	void SetMergedMaterialAtlasResolution(TSharedPtr<FString> NewSelection, ESelectInfo::Type SelectInfo);
 	
 	/** Called when actors selection is changed */
 	void OnActorSelectionChanged(const TArray<UObject*>& NewSelection, bool bForceRefresh=false);
@@ -85,4 +104,7 @@ private:
 	/**  */
 	TArray<TSharedPtr<FString>>	LightMapResolutionOptions;
 	TArray<TSharedPtr<FString>>	LightMapChannelOptions;
+	
+	/**  */
+	TArray<TSharedPtr<FString>>	MergedMaterialResolutionOptions;
 };
