#Fixed Merge Actor to use LongPackage Name for package name

diff --git a/Engine/Source/Developer/MeshUtilities/Private/MeshUtilities.cpp b/Engine/Source/Developer/MeshUtilities/Private/MeshUtilities.cpp
--- a/Engine/Source/Developer/MeshUtilities/Private/MeshUtilities.cpp
+++ b/Engine/Source/Developer/MeshUtilities/Private/MeshUtilities.cpp
@@ -4514,7 +4514,7 @@ void FMeshUtilities::MergeActors(
 		else
 		{
 			AssetName = FPackageName::GetShortName(InPackageName);
-			PackageName = InPackageName;
+			PackageName = FPackageName::GetLongPackagePath(InPackageName);
 		}
 
 		UPackage* Package = CreatePackage(NULL, *PackageName);
