#HLOD

- Made sure saving/loading works fine
- Packages are not created internally for merge actor if Outer exists - in this case, map
- It won't be visible in CB

diff --git a/Engine/Source/Developer/MeshUtilities/Private/MeshUtilities.cpp b/Engine/Source/Developer/MeshUtilities/Private/MeshUtilities.cpp
--- a/Engine/Source/Developer/MeshUtilities/Private/MeshUtilities.cpp
+++ b/Engine/Source/Developer/MeshUtilities/Private/MeshUtilities.cpp
@@ -131,7 +131,8 @@ private:
 	virtual void MergeActors(
 		const TArray<AActor*>& SourceActors,
 		const FMeshMergingSettings& InSettings,
-		const FString& PackageName,
+		UPackage* InOuter,
+		const FString& BasePackageName,
 		TArray<UObject*>& 
 		OutAssetsToSync, 
 		FVector& OutMergedActorLocation, 
@@ -4290,7 +4291,8 @@ static void MergeMaterials(UWorld* InWorld, const TArray<UMaterialInterface*>& I
 void FMeshUtilities::MergeActors(
 	const TArray<AActor*>& SourceActors, 
 	const FMeshMergingSettings& InSettings, 
-	const FString& InPackageName, 
+	UPackage* InOuter,
+	const FString& InBasePackageName,
 	TArray<UObject*>& OutAssetsToSync, 
 	FVector& OutMergedActorLocation, 
 	bool bSilent) const
@@ -4506,21 +4508,25 @@ void FMeshUtilities::MergeActors(
 	{
 		FString AssetName;
 		FString PackageName;
-		if (InPackageName.IsEmpty())
+		if (InBasePackageName.IsEmpty())
 		{
 			AssetName = TEXT("SM_MERGED_") + FPackageName::GetShortName(MergedMesh.AssetPackageName);
 			PackageName = FPackageName::GetLongPackagePath(MergedMesh.AssetPackageName) + TEXT("/") + AssetName;
 		}
 		else
 		{
-			AssetName = FPackageName::GetShortName(InPackageName);
-			PackageName = FPackageName::GetLongPackagePath(InPackageName);
+			AssetName = FPackageName::GetShortName(InBasePackageName);
+			PackageName =InOuter ? TEXT("") : FPackageName::GetLongPackagePath(InBasePackageName) + TEXT("/");
 		}
 
-		UPackage* Package = CreatePackage(NULL, *PackageName);
-		check(Package);
-		Package->FullyLoad();
-		Package->Modify();
+		UPackage* Package = InOuter;
+		if(Package == nullptr)
+		{
+			Package = CreatePackage(NULL, *PackageName);
+			check(Package);
+			Package->FullyLoad();
+			Package->Modify();
+		}
 
 		auto StaticMesh = NewObject<UStaticMesh>(Package, *AssetName, RF_Public | RF_Standalone);
 		StaticMesh->InitResources();
@@ -4554,21 +4560,25 @@ void FMeshUtilities::MergeActors(
 		{
 			FString MaterialAssetName;
 			FString MaterialPackageName;
-			if (InPackageName.IsEmpty())
+			if (InBasePackageName.IsEmpty())
 			{
 				MaterialAssetName = TEXT("M_MERGED_") + FPackageName::GetShortName(MergedMesh.AssetPackageName);
 				MaterialPackageName = FPackageName::GetLongPackagePath(MergedMesh.AssetPackageName) + TEXT("/") + MaterialAssetName;
 			}
 			else
 			{
-				MaterialAssetName = FPackageName::GetShortName(InPackageName);
-				MaterialPackageName = InPackageName;
+				MaterialAssetName = FPackageName::GetShortName(InBasePackageName);
+				MaterialPackageName =InOuter ? TEXT("") : FPackageName::GetLongPackagePath(InBasePackageName) + TEXT("/");
 			}
 
-			UPackage* MaterialPackage = CreatePackage(NULL, *MaterialPackageName);
-			check(MaterialPackage);
-			MaterialPackage->FullyLoad();
-			MaterialPackage->Modify();
+			UPackage* MaterialPackage = InOuter;
+			if(MaterialPackage == nullptr)
+			{
+				CreatePackage(NULL, *MaterialPackageName);
+				check(MaterialPackage);
+				MaterialPackage->FullyLoad();
+				MaterialPackage->Modify();
+			}
 
 			UMaterial* MergedMaterial = CreateMaterial(MergedFlatMaterial, MaterialPackage, MaterialAssetName, RF_Public|RF_Standalone, OutAssetsToSync);
 			UniqueMaterials.Add(MergedMaterial);
diff --git a/Engine/Source/Developer/MeshUtilities/Public/MeshUtilities.h b/Engine/Source/Developer/MeshUtilities/Public/MeshUtilities.h
--- a/Engine/Source/Developer/MeshUtilities/Public/MeshUtilities.h
+++ b/Engine/Source/Developer/MeshUtilities/Public/MeshUtilities.h
@@ -238,7 +238,8 @@ public:
 	virtual void MergeActors(
 		const TArray<AActor*>& SourceActors,
 		const FMeshMergingSettings& InSettings,
-		const FString& PackageName, 
+		UPackage* InOuter,
+		const FString& BasePackageName,
 		TArray<UObject*>& OutAssetsToSync, 
 		FVector& OutMergedActorLocation, 
 		bool bSilent=false) const = 0;
diff --git a/Engine/Source/Editor/UnrealEd/Private/Dialogs/SMeshMergingDialog.cpp b/Engine/Source/Editor/UnrealEd/Private/Dialogs/SMeshMergingDialog.cpp
--- a/Engine/Source/Editor/UnrealEd/Private/Dialogs/SMeshMergingDialog.cpp
+++ b/Engine/Source/Editor/UnrealEd/Private/Dialogs/SMeshMergingDialog.cpp
@@ -586,7 +586,7 @@ void SMeshMergingDialog::RunMerging()
 		
 	FVector MergedActorLocation;
 	TArray<UObject*> AssetsToSync;
-	MeshUtilities.MergeActors(Actors, MergingSettings, MergedMeshPackageName, AssetsToSync, MergedActorLocation);
+	MeshUtilities.MergeActors(Actors, MergingSettings, NULL, MergedMeshPackageName, AssetsToSync, MergedActorLocation);
 
 	if (AssetsToSync.Num())
 	{
diff --git a/Engine/Source/Editor/UnrealEd/Private/HierarchicalLOD.cpp b/Engine/Source/Editor/UnrealEd/Private/HierarchicalLOD.cpp
--- a/Engine/Source/Editor/UnrealEd/Private/HierarchicalLOD.cpp
+++ b/Engine/Source/Editor/UnrealEd/Private/HierarchicalLOD.cpp
@@ -614,13 +614,13 @@ void FLODCluster::BuildActor(class UWorld* InWorld, class ULevel* InLevel, const
 				{
 					FMeshMergingSettings MergeSetting;
 					MergeSetting.bMergeMaterials = true;
-					MeshUtilities.MergeActors(Actors, MergeSetting, AssetsPath + PackageName, OutAssets, OutProxyLocation, true );
+					MeshUtilities.MergeActors(Actors, MergeSetting, AssetsOuter, PackageName, OutAssets, OutProxyLocation, true );
 				}
 
-				// remove flags of RF_Public since we don't want these assets to be public
+				// this is ideally the way we should do, but 
 				for (auto& AssetIter : OutAssets)
 				{
-					AssetIter->ClearFlags(RF_Public);
+					AssetIter->ClearFlags(RF_Public | RF_Standalone);
 				}
 
 				UStaticMesh* MainMesh=NULL;
