#HLOD
- Added build of UseLOD
- This allows user to merge actor using certain LOD as base

diff --git a/Engine/Source/Developer/MeshUtilities/Private/MeshUtilities.cpp b/Engine/Source/Developer/MeshUtilities/Private/MeshUtilities.cpp
--- a/Engine/Source/Developer/MeshUtilities/Private/MeshUtilities.cpp
+++ b/Engine/Source/Developer/MeshUtilities/Private/MeshUtilities.cpp
@@ -133,8 +133,8 @@ private:
 		const FMeshMergingSettings& InSettings,
 		UPackage* InOuter,
 		const FString& BasePackageName,
-		TArray<UObject*>& 
-		OutAssetsToSync, 
+		int32 UseLOD, // does not build all LODs but only use this LOD to create base mesh
+		TArray<UObject*>& OutAssetsToSync, 
 		FVector& OutMergedActorLocation, 
 		bool bSilent=false) const override;
 	
@@ -4293,6 +4293,7 @@ void FMeshUtilities::MergeActors(
 	const FMeshMergingSettings& InSettings, 
 	UPackage* InOuter,
 	const FString& InBasePackageName,
+	int32 UseLOD, // does not build all LODs but only use this LOD to create base mesh
 	TArray<UObject*>& OutAssetsToSync, 
 	FVector& OutMergedActorLocation, 
 	bool bSilent) const
@@ -4342,35 +4343,74 @@ void FMeshUtilities::MergeActors(
 	
 	NumMaxLOD = FMath::Min(NumMaxLOD, MAX_STATIC_MESH_LODS);
 		
-	for (int32 LODIndex = 0; LODIndex < NumMaxLOD; ++LODIndex)
+	if ( UseLOD >=  0 && UseLOD < NumMaxLOD )
 	{
-		for (int32 MeshId = 0; MeshId < ComponentsToMerge.Num(); ++MeshId)
+		// change NumMaxLOD to be 1. We won't need anything else anymore
+		NumMaxLOD = 1;
+
+		for(int32 MeshId = 0; MeshId < ComponentsToMerge.Num(); ++MeshId)
 		{
 			UStaticMeshComponent* MeshComponent = ComponentsToMerge[MeshId];
-			
+
 			TArray<int32> MeshMaterialMap;
-			FRawMesh& RawMeshLOD = SourceMeshes[MeshId].MeshLOD[LODIndex];
+			FRawMesh& RawMeshLOD = SourceMeshes[MeshId].MeshLOD[0];
 
 			// We duplicate lower LOD in case this mesh has no LOD we want
-			int32 ExportLODIndex = FMath::Min(LODIndex, MeshComponent->StaticMesh->SourceModels.Num()-1);
+			int32 ExportLODIndex = FMath::Min(UseLOD, MeshComponent->StaticMesh->SourceModels.Num()-1);
 
-			if (ConstructRawMesh(MeshComponent, ExportLODIndex, RawMeshLOD, UniqueMaterials, MeshMaterialMap))
+			if(ConstructRawMesh(MeshComponent, ExportLODIndex, RawMeshLOD, UniqueMaterials, MeshMaterialMap))
 			{
-				MaterialMap.Add(FMeshIdAndLOD(MeshId, LODIndex), MeshMaterialMap);
-				
+				MaterialMap.Add(FMeshIdAndLOD(MeshId, 0), MeshMaterialMap);
+
 				// Should we use vertex colors?
-				if (InSettings.bImportVertexColors)
+				if(InSettings.bImportVertexColors)
 				{
 					// Propagate painted vertex colors into our raw mesh
-					PropagatePaintedColorsToRawMesh(MeshComponent, LODIndex, RawMeshLOD);
+					PropagatePaintedColorsToRawMesh(MeshComponent, 0, RawMeshLOD);
 					// Whether at least one of the meshes has vertex colors
-					bWithVertexColors[LODIndex]|= (RawMeshLOD.WedgeColors.Num() != 0);
+					bWithVertexColors[0]|= (RawMeshLOD.WedgeColors.Num() != 0);
 				}
-			
+
 				// Which UV channels has data at least in one mesh
-				for (int32 ChannelIdx = 0; ChannelIdx < MAX_MESH_TEXTURE_COORDS; ++ChannelIdx)
+				for(int32 ChannelIdx = 0; ChannelIdx < MAX_MESH_TEXTURE_COORDS; ++ChannelIdx)
 				{
-					bOcuppiedUVChannels[LODIndex][ChannelIdx]|= (RawMeshLOD.WedgeTexCoords[ChannelIdx].Num() != 0);
+					bOcuppiedUVChannels[0][ChannelIdx]|= (RawMeshLOD.WedgeTexCoords[ChannelIdx].Num() != 0);
+				}
+			}
+		}
+	}
+	else
+	{
+		for(int32 LODIndex = 0; LODIndex < NumMaxLOD; ++LODIndex)
+		{
+			for(int32 MeshId = 0; MeshId < ComponentsToMerge.Num(); ++MeshId)
+			{
+				UStaticMeshComponent* MeshComponent = ComponentsToMerge[MeshId];
+
+				TArray<int32> MeshMaterialMap;
+				FRawMesh& RawMeshLOD = SourceMeshes[MeshId].MeshLOD[LODIndex];
+
+				// We duplicate lower LOD in case this mesh has no LOD we want
+				int32 ExportLODIndex = FMath::Min(LODIndex, MeshComponent->StaticMesh->SourceModels.Num()-1);
+
+				if(ConstructRawMesh(MeshComponent, ExportLODIndex, RawMeshLOD, UniqueMaterials, MeshMaterialMap))
+				{
+					MaterialMap.Add(FMeshIdAndLOD(MeshId, LODIndex), MeshMaterialMap);
+
+					// Should we use vertex colors?
+					if(InSettings.bImportVertexColors)
+					{
+						// Propagate painted vertex colors into our raw mesh
+						PropagatePaintedColorsToRawMesh(MeshComponent, LODIndex, RawMeshLOD);
+						// Whether at least one of the meshes has vertex colors
+						bWithVertexColors[LODIndex]|= (RawMeshLOD.WedgeColors.Num() != 0);
+					}
+
+					// Which UV channels has data at least in one mesh
+					for(int32 ChannelIdx = 0; ChannelIdx < MAX_MESH_TEXTURE_COORDS; ++ChannelIdx)
+					{
+						bOcuppiedUVChannels[LODIndex][ChannelIdx]|= (RawMeshLOD.WedgeTexCoords[ChannelIdx].Num() != 0);
+					}
 				}
 			}
 		}
diff --git a/Engine/Source/Developer/MeshUtilities/Public/MeshUtilities.h b/Engine/Source/Developer/MeshUtilities/Public/MeshUtilities.h
--- a/Engine/Source/Developer/MeshUtilities/Public/MeshUtilities.h
+++ b/Engine/Source/Developer/MeshUtilities/Public/MeshUtilities.h
@@ -231,7 +231,10 @@ public:
 	 *
 	 * @param SourceActors				List of actors to merge
 	 * @param InSettings				Settings to use
-	 * @param PackageName				Destination package name for a generated assets
+	 * @param InOuter					Outer if required
+	 * @param InBasePackageName			Destination package name for a generated assets. Used if Outer is null. 
+	 * @param UseLOD					-1 if you'd like to build for all LODs. If you specify, that LOD mesh for source meshes will be used to merge the mesh
+	 *									This is used by hierarchical building LODs
 	 * @param OutAssetsToSync			Merged mesh assets
 	 * @param OutMergedActorLocation	World position of merged mesh
 	 */
@@ -239,7 +242,8 @@ public:
 		const TArray<AActor*>& SourceActors,
 		const FMeshMergingSettings& InSettings,
 		UPackage* InOuter,
-		const FString& BasePackageName,
+		const FString& InBasePackageName,
+		int32 UseLOD, // does not build all LODs but only use this LOD to create base mesh
 		TArray<UObject*>& OutAssetsToSync, 
 		FVector& OutMergedActorLocation, 
 		bool bSilent=false) const = 0;
diff --git a/Engine/Source/Editor/UnrealEd/Private/Dialogs/SMeshMergingDialog.cpp b/Engine/Source/Editor/UnrealEd/Private/Dialogs/SMeshMergingDialog.cpp
--- a/Engine/Source/Editor/UnrealEd/Private/Dialogs/SMeshMergingDialog.cpp
+++ b/Engine/Source/Editor/UnrealEd/Private/Dialogs/SMeshMergingDialog.cpp
@@ -586,7 +586,7 @@ void SMeshMergingDialog::RunMerging()
 		
 	FVector MergedActorLocation;
 	TArray<UObject*> AssetsToSync;
-	MeshUtilities.MergeActors(Actors, MergingSettings, NULL, MergedMeshPackageName, AssetsToSync, MergedActorLocation);
+	MeshUtilities.MergeActors(Actors, MergingSettings, NULL, MergedMeshPackageName, -1, AssetsToSync, MergedActorLocation);
 
 	if (AssetsToSync.Num())
 	{
diff --git a/Engine/Source/Editor/UnrealEd/Private/HierarchicalLOD.cpp b/Engine/Source/Editor/UnrealEd/Private/HierarchicalLOD.cpp
--- a/Engine/Source/Editor/UnrealEd/Private/HierarchicalLOD.cpp
+++ b/Engine/Source/Editor/UnrealEd/Private/HierarchicalLOD.cpp
@@ -614,7 +614,7 @@ void FLODCluster::BuildActor(class UWorld* InWorld, class ULevel* InLevel, const
 				{
 					FMeshMergingSettings MergeSetting;
 					MergeSetting.bMergeMaterials = true;
-					MeshUtilities.MergeActors(Actors, MergeSetting, AssetsOuter, PackageName, OutAssets, OutProxyLocation, true );
+					MeshUtilities.MergeActors(Actors, MergeSetting, AssetsOuter, PackageName, LODIdx+1, OutAssets, OutProxyLocation, true );
 				}
 
 				// this is ideally the way we should do, but 
