#HLOD

Fixed material flattening UV bug where input is out of range, it doesn't work.
Validated the input to be in the right range.

diff --git a/Engine/Source/Developer/MeshUtilities/Private/MeshUtilities.cpp b/Engine/Source/Developer/MeshUtilities/Private/MeshUtilities.cpp
--- a/Engine/Source/Developer/MeshUtilities/Private/MeshUtilities.cpp
+++ b/Engine/Source/Developer/MeshUtilities/Private/MeshUtilities.cpp
@@ -4154,6 +4154,34 @@ struct FRawMeshExt
 	FVector		Pivot;	
 };
 
+static FVector2D GetValidUV(FVector2D& UV)
+{
+	// first make sure they're positive
+	FVector2D NewUV;
+	if (UV.X < 0.f)
+	{
+		NewUV.X = UV.X + FMath::CeilToInt(FMath::Abs(UV.X));
+	}
+	else
+	{
+		NewUV.X = UV.X;
+	}
+	
+	if (UV.Y < 0.f)
+	{
+		NewUV.Y = UV.Y + FMath::CeilToInt(FMath::Abs(UV.Y));
+	}
+	else
+	{
+		NewUV.Y = UV.Y;
+	}
+
+	// now make sure they're within [0, 1]
+	NewUV.X = FMath::Fmod(NewUV.X, 1.f);
+	NewUV.Y = FMath::Fmod(NewUV.Y, 1.f);
+	return NewUV;
+}
+
 static void MergeMaterials(UWorld* InWorld, const TArray<UMaterialInterface*>& InMaterialList, TArray<FRawMeshExt>& InOutMeshes, MaterialExportUtils::FFlattenMaterial& OutMergedMaterial)
 {
 	using namespace MaterialExportUtils;
@@ -4273,9 +4301,12 @@ static void MergeMaterials(UWorld* InWorld, const TArray<UMaterialInterface*>& I
 						int32 FaceMaterialIndex = RawMesh.FaceMaterialIndices[i/3];
 						FVector2D UVOffset = UVOffsets[FaceMaterialIndex];
 
-						UVs[i+0] = UVs[i+0]*UVScale + UVOffset;
-						UVs[i+1] = UVs[i+1]*UVScale + UVOffset;
-						UVs[i+2] = UVs[i+2]*UVScale + UVOffset;
+						FVector2D UV0 = GetValidUV(UVs[i+0]);
+						FVector2D UV1 = GetValidUV(UVs[i+1]);
+						FVector2D UV2 = GetValidUV(UVs[i+2]);
+						UVs[i+0] = UV0 *UVScale + UVOffset;
+						UVs[i+1] = UV1 *UVScale + UVOffset;
+						UVs[i+2] = UV2 *UVScale + UVOffset;
 					}
 				}
 			}
