#HLOD: Fixed crash with merge actor feature

diff --git a/Engine/Source/Developer/MeshUtilities/Private/MeshUtilities.cpp b/Engine/Source/Developer/MeshUtilities/Private/MeshUtilities.cpp
--- a/Engine/Source/Developer/MeshUtilities/Private/MeshUtilities.cpp
+++ b/Engine/Source/Developer/MeshUtilities/Private/MeshUtilities.cpp
@@ -4587,7 +4587,7 @@ void FMeshUtilities::MergeActors(
 		else
 		{
 			AssetName = FPackageName::GetShortName(InBasePackageName);
-			PackageName =InOuter ? TEXT("") : FPackageName::GetLongPackagePath(InBasePackageName) + TEXT("/");
+			PackageName =InOuter ? TEXT("") : InBasePackageName + TEXT("/");
 		}
 
 		UPackage* Package = InOuter;
@@ -4639,13 +4639,13 @@ void FMeshUtilities::MergeActors(
 			else
 			{
 				MaterialAssetName = FPackageName::GetShortName(InBasePackageName);
-				MaterialPackageName =InOuter ? TEXT("") : FPackageName::GetLongPackagePath(InBasePackageName) + TEXT("/");
+				MaterialPackageName =InOuter ? TEXT("") : InBasePackageName + TEXT("/");
 			}
 
 			UPackage* MaterialPackage = InOuter;
 			if(MaterialPackage == nullptr)
 			{
-				CreatePackage(NULL, *MaterialPackageName);
+				MaterialPackage = CreatePackage(NULL, *MaterialPackageName);
 				check(MaterialPackage);
 				MaterialPackage->FullyLoad();
 				MaterialPackage->Modify();
