Fixed: Merge materails will no longer try to merge non-opaque materials
Fixed: Merge materials UV wrapping was changing UVs [1,1] to [0,0]

diff --git a/Engine/Source/Developer/MeshUtilities/Private/MeshUtilities.cpp b/Engine/Source/Developer/MeshUtilities/Private/MeshUtilities.cpp
--- a/Engine/Source/Developer/MeshUtilities/Private/MeshUtilities.cpp
+++ b/Engine/Source/Developer/MeshUtilities/Private/MeshUtilities.cpp
@@ -4154,39 +4154,75 @@ struct FRawMeshExt
 	FVector		Pivot;	
 };
 
-static FVector2D GetValidUV(FVector2D& UV)
+struct FRawMeshUVTransform
 {
+	FVector2D Offset;
+	FVector2D Scale;
+	
+	bool IsValid() const
+	{
+		return (Scale != FVector2D::ZeroVector);
+	}
+};
+
+static FVector2D GetValidUV(const FVector2D& UV)
+{
+	FVector2D NewUV = UV;
 	// first make sure they're positive
-	FVector2D NewUV;
-	if (UV.X < 0.f)
+	if (UV.X < 0.0f)
 	{
 		NewUV.X = UV.X + FMath::CeilToInt(FMath::Abs(UV.X));
 	}
-	else
+	
+	if (UV.Y < 0.0f)
 	{
-		NewUV.X = UV.X;
+		NewUV.Y = UV.Y + FMath::CeilToInt(FMath::Abs(UV.Y));
 	}
 	
-	if (UV.Y < 0.f)
+	// now make sure they're within [0, 1]
+	if (UV.X > 1.0f)
 	{
-		NewUV.Y = UV.Y + FMath::CeilToInt(FMath::Abs(UV.Y));
+		NewUV.X = FMath::Fmod(NewUV.X, 1.0f);
 	}
-	else
+	
+	if (UV.Y > 1.0f)
 	{
-		NewUV.Y = UV.Y;
+		NewUV.Y = FMath::Fmod(NewUV.Y, 1.0f);
 	}
-
-	// now make sure they're within [0, 1]
-	NewUV.X = FMath::Fmod(NewUV.X, 1.f);
-	NewUV.Y = FMath::Fmod(NewUV.Y, 1.f);
+	
 	return NewUV;
 }
 
-static void MergeMaterials(UWorld* InWorld, const TArray<UMaterialInterface*>& InMaterialList, TArray<FRawMeshExt>& InOutMeshes, MaterialExportUtils::FFlattenMaterial& OutMergedMaterial)
+static void MergeMaterials(UWorld* InWorld, const TArray<UMaterialInterface*>& InMaterialList, MaterialExportUtils::FFlattenMaterial& OutMergedMaterial, TArray<FRawMeshUVTransform>& OutUVTransforms)
 {
 	using namespace MaterialExportUtils;
 	
-	int32 AtlasGridSize = FMath::CeilToInt(FMath::Sqrt(InMaterialList.Num()));
+	OutUVTransforms.Reserve(InMaterialList.Num());
+
+	// We support merging only for opaque materials
+	int32 NumOpaqueMaterials = 0;
+	// Fill output UV transforms with invalid values
+	for (auto Material : InMaterialList)
+	{
+		if (Material->GetBlendMode() == BLEND_Opaque)
+		{
+			NumOpaqueMaterials++;
+		}
+		
+		// Invalid UV transform
+		FRawMeshUVTransform UVTransform;
+		UVTransform.Offset = FVector2D::ZeroVector;
+		UVTransform.Scale = FVector2D::ZeroVector;
+		OutUVTransforms.Add(UVTransform);
+	}
+
+	if (NumOpaqueMaterials == 0)
+	{
+		// Nothing to merge
+		return;
+	}
+		
+	int32 AtlasGridSize = FMath::CeilToInt(FMath::Sqrt(NumOpaqueMaterials));
 	FIntPoint AtlasTextureSize = OutMergedMaterial.DiffuseSize;
 	FIntPoint ExportTextureSize	= AtlasTextureSize/AtlasGridSize;
 	int32 AtlasNumSamples = AtlasTextureSize.X*AtlasTextureSize.Y;
@@ -4227,19 +4263,24 @@ static void MergeMaterials(UWorld* InWorld, const TArray<UMaterialInterface*>& I
 	int32 AtlasRowIdx = 0;
 	int32 AtlasColIdx = 0;
 	FIntPoint AtlasTargetPos = FIntPoint(0,0);
-	TArray<FVector2D> UVOffsets;
 					
 	// Flatten all materials and merge them into one material using texture atlases
-	for (const UMaterialInterface* Material : InMaterialList)
+	for (int32 MatIdx = 0; MatIdx < InMaterialList.Num(); ++MatIdx)
 	{
+		UMaterialInterface* Material = InMaterialList[MatIdx];
+		if (Material->GetBlendMode() != BLEND_Opaque)
+		{
+			continue;
+		}
+		
 		FFlattenMaterial FlatMaterial;
 		FlatMaterial.DiffuseSize	= ExportTextureSize;
 		FlatMaterial.NormalSize		= bExportNormal ? ExportTextureSize : FIntPoint::ZeroValue;
 		FlatMaterial.MetallicSize	= bExportMetallic ? ExportTextureSize : FIntPoint::ZeroValue;
-		FlatMaterial.RoughnessSize	= bExportRoughness ? ExportTextureSize : FIntPoint::ZeroValue;;
-		FlatMaterial.SpecularSize	= bExportSpecular ? ExportTextureSize : FIntPoint::ZeroValue;;
+		FlatMaterial.RoughnessSize	= bExportRoughness ? ExportTextureSize : FIntPoint::ZeroValue;
+		FlatMaterial.SpecularSize	= bExportSpecular ? ExportTextureSize : FIntPoint::ZeroValue;
 
-		ExportMaterial(InWorld, const_cast<UMaterialInterface*>(Material), FlatMaterial);
+		ExportMaterial(InWorld, Material, FlatMaterial);
 
 		if (FlatMaterial.DiffuseSamples.Num() > 0)
 		{
@@ -4265,10 +4306,16 @@ static void MergeMaterials(UWorld* InWorld, const TArray<UMaterialInterface*>& I
 		{
 			CopyTextureRect(FlatMaterial.SpecularSamples.GetData(), ExportTextureSize, OutMergedMaterial.SpecularSamples.GetData(), AtlasTextureSize, AtlasTargetPos);
 		}
+		
+		check(OutUVTransforms.IsValidIndex(MatIdx));
 
-		UVOffsets.Add(FVector2D(
+		OutUVTransforms[MatIdx].Offset = FVector2D(
 			(float)AtlasTargetPos.X/AtlasTextureSize.X, 
-			(float)AtlasTargetPos.Y/AtlasTextureSize.Y));
+			(float)AtlasTargetPos.Y/AtlasTextureSize.Y);
+		
+		OutUVTransforms[MatIdx].Scale = FVector2D(
+			(float)ExportTextureSize.X/AtlasTextureSize.X, 
+			(float)ExportTextureSize.Y/AtlasTextureSize.Y);
 		
 		AtlasColIdx++;
 		if (AtlasColIdx >= AtlasGridSize)
@@ -4279,44 +4326,6 @@ static void MergeMaterials(UWorld* InWorld, const TArray<UMaterialInterface*>& I
 
 		AtlasTargetPos = FIntPoint(AtlasColIdx*ExportTextureSize.X, AtlasRowIdx*ExportTextureSize.Y);
 	}
-
-	// Adjust UVs
-	FVector2D UVScale = FVector2D(
-		(float)ExportTextureSize.X/AtlasTextureSize.X, 
-		(float)ExportTextureSize.Y/AtlasTextureSize.Y);
-
-	for (int32 MeshIndex = 0; MeshIndex < InOutMeshes.Num(); ++MeshIndex)
-	{
-		for (int32 LODIndex = 0; LODIndex < MAX_STATIC_MESH_LODS; ++LODIndex)
-		{
-			FRawMesh& RawMesh = InOutMeshes[MeshIndex].MeshLOD[LODIndex];
-
-			for (int32 UVChannelIdx = 0; UVChannelIdx < MAX_MESH_TEXTURE_COORDS; ++UVChannelIdx)
-			{
-				TArray<FVector2D>& UVs = RawMesh.WedgeTexCoords[UVChannelIdx];
-				if (UVs.Num() > 0)
-				{
-					for (int32 i = 0; i < UVs.Num(); i+=3)
-					{
-						int32 FaceMaterialIndex = RawMesh.FaceMaterialIndices[i/3];
-						FVector2D UVOffset = UVOffsets[FaceMaterialIndex];
-
-						FVector2D UV0 = GetValidUV(UVs[i+0]);
-						FVector2D UV1 = GetValidUV(UVs[i+1]);
-						FVector2D UV2 = GetValidUV(UVs[i+2]);
-						UVs[i+0] = UV0 *UVScale + UVOffset;
-						UVs[i+1] = UV1 *UVScale + UVOffset;
-						UVs[i+2] = UV2 *UVScale + UVOffset;
-					}
-				}
-			}
-			
-			for (int32& FaceMaterialIndex : RawMesh.FaceMaterialIndices)
-			{
-				FaceMaterialIndex = 0;
-			}
-		}
-	}
 }
 
 void FMeshUtilities::MergeActors(
@@ -4470,30 +4479,108 @@ void FMeshUtilities::MergeActors(
 		}
 	}
 
-	MaterialExportUtils::FFlattenMaterial MergedFlatMaterial;
-	
+	FRawMeshExt MergedMesh;
+	// Use first mesh for naming and pivot
+	MergedMesh.AssetPackageName = SourceMeshes[0].AssetPackageName;
+	MergedMesh.Pivot = InSettings.bPivotPointAtZero ? FVector::ZeroVector : SourceMeshes[0].Pivot;
+			
 	if (InSettings.bMergeMaterials)
 	{
 		FIntPoint AtlasTextureSize			= FIntPoint(InSettings.MergedMaterialAtlasResolution, InSettings.MergedMaterialAtlasResolution);
+		MaterialExportUtils::FFlattenMaterial MergedFlatMaterial;
 		MergedFlatMaterial.DiffuseSize		= AtlasTextureSize;
 		MergedFlatMaterial.NormalSize		= InSettings.bExportNormalMap ? AtlasTextureSize : FIntPoint::ZeroValue;
 		MergedFlatMaterial.MetallicSize		= InSettings.bExportMetallicMap ? AtlasTextureSize : FIntPoint::ZeroValue;
 		MergedFlatMaterial.RoughnessSize	= InSettings.bExportRoughnessMap ? AtlasTextureSize : FIntPoint::ZeroValue;
 		MergedFlatMaterial.SpecularSize		= InSettings.bExportSpecularMap ? AtlasTextureSize : FIntPoint::ZeroValue;
-			
+		TArray<FRawMeshUVTransform> UVTransforms;
 		UWorld* World = SourceActors[0]->GetWorld();
-		MergeMaterials(World, UniqueMaterials, SourceMeshes, MergedFlatMaterial);
-		
-		// We assume all materials was merged into MergedFlatMaterial
-		UniqueMaterials.Empty();
-	}
 
-	FRawMeshExt MergedMesh;
+		MergeMaterials(World, UniqueMaterials, MergedFlatMaterial, UVTransforms);
 
-	// Use first mesh for naming and pivot
-	MergedMesh.AssetPackageName = SourceMeshes[0].AssetPackageName;
-	MergedMesh.Pivot = InSettings.bPivotPointAtZero ? FVector::ZeroVector : SourceMeshes[0].Pivot;
-	
+		// Remove merged materials from materials list and remap mesh materials
+		TArray<int32> MaterialsRemapTable;
+		TArray<UMaterialInterface*>	RearangedMaterials;
+		check(UniqueMaterials.Num() == UVTransforms.Num());
+		for (int32 MatIdx = 0; MatIdx < UniqueMaterials.Num(); ++MatIdx)
+		{
+			int32 RearangedMatIdx = INDEX_NONE;
+			if (!UVTransforms[MatIdx].IsValid())
+			{
+				RearangedMatIdx = RearangedMaterials.Add(UniqueMaterials[MatIdx]);
+			}
+			MaterialsRemapTable.Add(RearangedMatIdx);
+		}
+		UniqueMaterials = RearangedMaterials;
+		int32 MergedMaterialIdx = UniqueMaterials.Num();
+
+		// Adjust UVs and remap material indices
+		for (int32 MeshIndex = 0; MeshIndex < SourceMeshes.Num(); ++MeshIndex)
+		{
+			for (int32 LODIndex = 0; LODIndex < MAX_STATIC_MESH_LODS; ++LODIndex)
+			{
+				FRawMesh& RawMesh = SourceMeshes[MeshIndex].MeshLOD[LODIndex];
+
+				for (int32 UVChannelIdx = 0; UVChannelIdx < MAX_MESH_TEXTURE_COORDS; ++UVChannelIdx)
+				{
+					TArray<FVector2D>& UVs = RawMesh.WedgeTexCoords[UVChannelIdx];
+					if (UVs.Num() > 0)
+					{
+						int32 UVIdx = 0; 
+						for (int32 FaceMaterialIndex : RawMesh.FaceMaterialIndices)
+						{
+							const FRawMeshUVTransform& UVTransform = UVTransforms[FaceMaterialIndex];
+							if (UVTransform.IsValid())
+							{
+								FVector2D UV0 = GetValidUV(UVs[UVIdx+0]);
+								FVector2D UV1 = GetValidUV(UVs[UVIdx+1]);
+								FVector2D UV2 = GetValidUV(UVs[UVIdx+2]);
+								UVs[UVIdx+0] = UV0 * UVTransform.Scale + UVTransform.Offset;
+								UVs[UVIdx+1] = UV1 * UVTransform.Scale + UVTransform.Offset;
+								UVs[UVIdx+2] = UV2 * UVTransform.Scale + UVTransform.Offset;
+							}
+						
+							UVIdx+=3;
+						}
+					}
+				}
+				
+				// Remap material indexes
+				for (int32& FaceMaterialIndex : RawMesh.FaceMaterialIndices)
+				{
+					int32 RemapedIdx = MaterialsRemapTable[FaceMaterialIndex];
+					FaceMaterialIndex = (RemapedIdx == INDEX_NONE ? MergedMaterialIdx : RemapedIdx);
+				}
+			}
+		}
+
+		// Create merged material asset
+		FString MaterialAssetName;
+		FString MaterialPackageName;
+		if (InBasePackageName.IsEmpty())
+		{
+			MaterialAssetName = TEXT("M_MERGED_") + FPackageName::GetShortName(MergedMesh.AssetPackageName);
+			MaterialPackageName = FPackageName::GetLongPackagePath(MergedMesh.AssetPackageName) + TEXT("/") + MaterialAssetName;
+		}
+		else
+		{
+			MaterialAssetName = TEXT("M_") + FPackageName::GetShortName(InBasePackageName);
+			MaterialPackageName = FPackageName::GetLongPackagePath(InBasePackageName) + TEXT("/") + MaterialAssetName;
+		}
+
+		UPackage* MaterialPackage = InOuter;
+		if(MaterialPackage == nullptr)
+		{
+			MaterialPackage = CreatePackage(nullptr, *MaterialPackageName);
+			check(MaterialPackage);
+			MaterialPackage->FullyLoad();
+			MaterialPackage->Modify();
+		}
+
+		UMaterial* MergedMaterial = CreateMaterial(MergedFlatMaterial, MaterialPackage, MaterialAssetName, RF_Public|RF_Standalone, OutAssetsToSync);
+		UniqueMaterials.Add(MergedMaterial);
+	}
+		
 	// Merge meshes into single mesh
 	for (int32 SourceMeshIdx = 0; SourceMeshIdx < SourceMeshes.Num(); ++SourceMeshIdx)
 	{
@@ -4587,7 +4674,7 @@ void FMeshUtilities::MergeActors(
 		else
 		{
 			AssetName = FPackageName::GetShortName(InBasePackageName);
-			PackageName =InOuter ? TEXT("") : InBasePackageName + TEXT("/");
+			PackageName = InBasePackageName;
 		}
 
 		UPackage* Package = InOuter;
@@ -4626,35 +4713,7 @@ void FMeshUtilities::MergeActors(
 				SrcModel->RawMeshBulkData->SaveRawMesh(MergedMeshLOD);
 			}
 		}
-
-		if (InSettings.bMergeMaterials)
-		{
-			FString MaterialAssetName;
-			FString MaterialPackageName;
-			if (InBasePackageName.IsEmpty())
-			{
-				MaterialAssetName = TEXT("M_MERGED_") + FPackageName::GetShortName(MergedMesh.AssetPackageName);
-				MaterialPackageName = FPackageName::GetLongPackagePath(MergedMesh.AssetPackageName) + TEXT("/") + MaterialAssetName;
-			}
-			else
-			{
-				MaterialAssetName = FPackageName::GetShortName(InBasePackageName);
-				MaterialPackageName =InOuter ? TEXT("") : InBasePackageName + TEXT("/");
-			}
-
-			UPackage* MaterialPackage = InOuter;
-			if(MaterialPackage == nullptr)
-			{
-				MaterialPackage = CreatePackage(NULL, *MaterialPackageName);
-				check(MaterialPackage);
-				MaterialPackage->FullyLoad();
-				MaterialPackage->Modify();
-			}
-
-			UMaterial* MergedMaterial = CreateMaterial(MergedFlatMaterial, MaterialPackage, MaterialAssetName, RF_Public|RF_Standalone, OutAssetsToSync);
-			UniqueMaterials.Add(MergedMaterial);
-		}
-	
+			
 		// Assign materials
 		for (UMaterialInterface* Material : UniqueMaterials)
 		{
