Exposed "Export specific LOD" to mesh merging dialog
Some cleanup for MeshMerging functionality

diff --git a/Engine/Source/Developer/MeshUtilities/MeshUtilities.Build.cs b/Engine/Source/Developer/MeshUtilities/MeshUtilities.Build.cs
--- a/Engine/Source/Developer/MeshUtilities/MeshUtilities.Build.cs
+++ b/Engine/Source/Developer/MeshUtilities/MeshUtilities.Build.cs
@@ -14,19 +14,11 @@ public class MeshUtilities : ModuleRules
 				"Engine",
 				"RawMesh",
 				"RenderCore", // For FPackedNormal
-				"Slate",      // For FSlateTextureAtlas
 				"SlateCore",
 				"UnrealEd",
-				"Landscape",
 			}
 		);
 
-		PrivateIncludePathModuleNames.AddRange(
-				new string[] {
-					"Landscape",
-				}
-		);
-
 		AddThirdPartyPrivateStaticDependencies(Target, "nvTriStrip");
 		AddThirdPartyPrivateStaticDependencies(Target, "ForsythTriOptimizer");
 		AddThirdPartyPrivateStaticDependencies(Target, "MeshSimplifier");
diff --git a/Engine/Source/Developer/MeshUtilities/Private/MeshUtilities.cpp b/Engine/Source/Developer/MeshUtilities/Private/MeshUtilities.cpp
--- a/Engine/Source/Developer/MeshUtilities/Private/MeshUtilities.cpp
+++ b/Engine/Source/Developer/MeshUtilities/Private/MeshUtilities.cpp
@@ -3,14 +3,12 @@
 #include "MeshUtilitiesPrivate.h"
 #include "StaticMeshResources.h"
 #include "SkeletalMeshTypes.h"
-#include "LandscapeRender.h"
 #include "MeshBuild.h"
 #include "TessellationRendering.h"
 #include "NvTriStrip.h"
 #include "forsythtriangleorderoptimizer.h"
 #include "ThirdParty/nvtesslib/inc/nvtess.h"
 #include "SkeletalMeshTools.h"
-#include "LandscapeDataAccess.h"
 #include "ImageUtils.h"
 #include "MaterialExportUtils.h"
 #include "Textures/TextureAtlas.h"
@@ -4365,11 +4363,11 @@ void FMeshUtilities::MergeActors(
 	bool											bWithVertexColors[MAX_STATIC_MESH_LODS] = {};
 	bool											bOcuppiedUVChannels[MAX_STATIC_MESH_LODS][MAX_MESH_TEXTURE_COORDS] = {};
 		
-	// Convert collected static mesh components into raw meshes
 	SourceMeshes.Reserve(ComponentsToMerge.Num());
 	SourceMeshes.SetNum(ComponentsToMerge.Num());
 
 	int32 NumMaxLOD = 0;
+	// Convert collected static mesh components into raw meshes
 	for (int32 MeshId = 0; MeshId < ComponentsToMerge.Num(); ++MeshId)
 	{
 		UStaticMeshComponent* MeshComponent = ComponentsToMerge[MeshId];
@@ -4380,77 +4378,47 @@ void FMeshUtilities::MergeActors(
 		// Source mesh asset package name
 		SourceMeshes[MeshId].AssetPackageName = MeshComponent->StaticMesh->GetOutermost()->GetName();
 	}
-	
+
 	NumMaxLOD = FMath::Min(NumMaxLOD, MAX_STATIC_MESH_LODS);
 		
-	if ( UseLOD >=  0 && UseLOD < NumMaxLOD )
+	int32 StartLODIndex = 0;
+	if (UseLOD >= 0)
+	{
+		// Will export only one specified LOD as base mesh
+		StartLODIndex = FMath::Min(UseLOD, NumMaxLOD - 1);
+		NumMaxLOD = StartLODIndex + 1;
+	}
+		
+	int32 RawMeshLODIdx = 0;
+	for(int32 LODIndex = StartLODIndex; LODIndex < NumMaxLOD; ++LODIndex, ++RawMeshLODIdx)
 	{
-		// change NumMaxLOD to be 1. We won't need anything else anymore
-		NumMaxLOD = 1;
-
 		for(int32 MeshId = 0; MeshId < ComponentsToMerge.Num(); ++MeshId)
 		{
 			UStaticMeshComponent* MeshComponent = ComponentsToMerge[MeshId];
 
 			TArray<int32> MeshMaterialMap;
-			FRawMesh& RawMeshLOD = SourceMeshes[MeshId].MeshLOD[0];
+			FRawMesh& RawMeshLOD = SourceMeshes[MeshId].MeshLOD[RawMeshLODIdx];
 
 			// We duplicate lower LOD in case this mesh has no LOD we want
-			int32 ExportLODIndex = FMath::Min(UseLOD, MeshComponent->StaticMesh->SourceModels.Num()-1);
+			int32 ExportLODIndex = FMath::Min(LODIndex, MeshComponent->StaticMesh->SourceModels.Num()-1);
 
 			if(ConstructRawMesh(MeshComponent, ExportLODIndex, RawMeshLOD, UniqueMaterials, MeshMaterialMap))
 			{
-				MaterialMap.Add(FMeshIdAndLOD(MeshId, 0), MeshMaterialMap);
+				MaterialMap.Add(FMeshIdAndLOD(MeshId, RawMeshLODIdx), MeshMaterialMap);
 
 				// Should we use vertex colors?
 				if(InSettings.bImportVertexColors)
 				{
 					// Propagate painted vertex colors into our raw mesh
-					PropagatePaintedColorsToRawMesh(MeshComponent, 0, RawMeshLOD);
+					PropagatePaintedColorsToRawMesh(MeshComponent, ExportLODIndex, RawMeshLOD);
 					// Whether at least one of the meshes has vertex colors
-					bWithVertexColors[0]|= (RawMeshLOD.WedgeColors.Num() != 0);
+					bWithVertexColors[RawMeshLODIdx]|= (RawMeshLOD.WedgeColors.Num() != 0);
 				}
 
 				// Which UV channels has data at least in one mesh
 				for(int32 ChannelIdx = 0; ChannelIdx < MAX_MESH_TEXTURE_COORDS; ++ChannelIdx)
 				{
-					bOcuppiedUVChannels[0][ChannelIdx]|= (RawMeshLOD.WedgeTexCoords[ChannelIdx].Num() != 0);
-				}
-			}
-		}
-	}
-	else
-	{
-		for(int32 LODIndex = 0; LODIndex < NumMaxLOD; ++LODIndex)
-		{
-			for(int32 MeshId = 0; MeshId < ComponentsToMerge.Num(); ++MeshId)
-			{
-				UStaticMeshComponent* MeshComponent = ComponentsToMerge[MeshId];
-
-				TArray<int32> MeshMaterialMap;
-				FRawMesh& RawMeshLOD = SourceMeshes[MeshId].MeshLOD[LODIndex];
-
-				// We duplicate lower LOD in case this mesh has no LOD we want
-				int32 ExportLODIndex = FMath::Min(LODIndex, MeshComponent->StaticMesh->SourceModels.Num()-1);
-
-				if(ConstructRawMesh(MeshComponent, ExportLODIndex, RawMeshLOD, UniqueMaterials, MeshMaterialMap))
-				{
-					MaterialMap.Add(FMeshIdAndLOD(MeshId, LODIndex), MeshMaterialMap);
-
-					// Should we use vertex colors?
-					if(InSettings.bImportVertexColors)
-					{
-						// Propagate painted vertex colors into our raw mesh
-						PropagatePaintedColorsToRawMesh(MeshComponent, LODIndex, RawMeshLOD);
-						// Whether at least one of the meshes has vertex colors
-						bWithVertexColors[LODIndex]|= (RawMeshLOD.WedgeColors.Num() != 0);
-					}
-
-					// Which UV channels has data at least in one mesh
-					for(int32 ChannelIdx = 0; ChannelIdx < MAX_MESH_TEXTURE_COORDS; ++ChannelIdx)
-					{
-						bOcuppiedUVChannels[LODIndex][ChannelIdx]|= (RawMeshLOD.WedgeTexCoords[ChannelIdx].Num() != 0);
-					}
+					bOcuppiedUVChannels[RawMeshLODIdx][ChannelIdx]|= (RawMeshLOD.WedgeTexCoords[ChannelIdx].Num() != 0);
 				}
 			}
 		}
diff --git a/Engine/Source/Editor/UnrealEd/Private/Dialogs/SMeshMergingDialog.cpp b/Engine/Source/Editor/UnrealEd/Private/Dialogs/SMeshMergingDialog.cpp
--- a/Engine/Source/Editor/UnrealEd/Private/Dialogs/SMeshMergingDialog.cpp
+++ b/Engine/Source/Editor/UnrealEd/Private/Dialogs/SMeshMergingDialog.cpp
@@ -23,6 +23,8 @@
 
 SMeshMergingDialog::SMeshMergingDialog()
 	: bPlaceInWorld(false)
+	, bExportSpecificLOD(false)
+	, ExportLODIndex(0)
 {
 }
 
@@ -59,6 +61,11 @@ void SMeshMergingDialog::Construct(const FArguments& InArgs)
 		LightMapChannelOptions.Add(MakeShareable(new FString(FString::FormatAsNumber(Index))));
 	}
 
+	for (int32 Index = 0; Index < MAX_STATIC_MESH_LODS; Index++)
+	{
+		ExportLODOptions.Add(MakeShareable(new FString(FString::FormatAsNumber(Index))));
+	}
+
 	GenerateNewPackageName();
 
 	// Create widget layout
@@ -150,6 +157,40 @@ void SMeshMergingDialog::Construct(const FArguments& InArgs)
 			.BorderImage(FEditorStyle::GetBrush("ToolPanel.GroupBorder"))
 			[
 				SNew(SVerticalBox)
+
+				// LOD to export
+				+SVerticalBox::Slot()
+				.AutoHeight()
+				.Padding(FEditorStyle::GetMargin("StandardDialog.ContentPadding"))
+				[
+					SNew(SHorizontalBox)
+
+					+SHorizontalBox::Slot()
+					.AutoWidth()
+					.VAlign(VAlign_Center)
+					[
+						SNew(SCheckBox)
+						.Type(ESlateCheckBoxType::CheckBox)
+						.IsChecked(this, &SMeshMergingDialog::GetExportSpecificLODEnabled)
+						.OnCheckStateChanged(this, &SMeshMergingDialog::SetExportSpecificLODEnabled)
+						.Content()
+						[
+							SNew(STextBlock).Text(LOCTEXT("TargetMeshLODIndexLabel", "Export specific LOD:"))
+						]
+					]
+					
+					+SHorizontalBox::Slot()
+					.AutoWidth()
+					.VAlign(VAlign_Center)
+					.Padding(4,0,4,0)
+					[
+						SNew(STextComboBox)
+						.IsEnabled(this, &SMeshMergingDialog::IsExportSpecificLODEnabled)
+						.OptionsSource(&ExportLODOptions)
+						.InitiallySelectedItem(ExportLODOptions[ExportLODIndex])
+						.OnSelectionChanged(this, &SMeshMergingDialog::SetExportSpecificLODIndex)
+					]
+				]
 					
 				// Vertex colors
 				+SVerticalBox::Slot()
@@ -416,6 +457,26 @@ void SMeshMergingDialog::SetTargetLightMapResolution(TSharedPtr<FString> NewSele
 	TTypeFromString<int32>::FromString(MergingSettings.TargetLightMapResolution, **NewSelection);
 }
 
+ECheckBoxState SMeshMergingDialog::GetExportSpecificLODEnabled() const
+{
+	return (bExportSpecificLOD ? ECheckBoxState::Checked : ECheckBoxState::Unchecked);
+}
+
+void SMeshMergingDialog::SetExportSpecificLODEnabled(ECheckBoxState NewValue)
+{
+	bExportSpecificLOD = (NewValue == ECheckBoxState::Checked);
+}
+
+bool SMeshMergingDialog::IsExportSpecificLODEnabled() const
+{
+	return bExportSpecificLOD;
+}
+
+void SMeshMergingDialog::SetExportSpecificLODIndex(TSharedPtr<FString> NewSelection, ESelectInfo::Type SelectInfo)
+{
+	TTypeFromString<int32>::FromString(ExportLODIndex, **NewSelection);
+}
+
 ECheckBoxState SMeshMergingDialog::GetImportVertexColors() const
 {
 	return (MergingSettings.bImportVertexColors ? ECheckBoxState::Checked : ECheckBoxState::Unchecked);
@@ -583,10 +644,12 @@ void SMeshMergingDialog::RunMerging()
 		OpenMsgDlgInt(EAppMsgType::Ok, Message, NSLOCTEXT("UnrealEd", "FailedToMergeActors_Title", "Unable to merge actors"));
 		return;
 	}
-		
+
+	int32 TargetMeshLOD = bExportSpecificLOD ? ExportLODIndex : INDEX_NONE;
+
 	FVector MergedActorLocation;
 	TArray<UObject*> AssetsToSync;
-	MeshUtilities.MergeActors(Actors, MergingSettings, NULL, MergedMeshPackageName, -1, AssetsToSync, MergedActorLocation);
+	MeshUtilities.MergeActors(Actors, MergingSettings, NULL, MergedMeshPackageName, TargetMeshLOD, AssetsToSync, MergedActorLocation);
 
 	if (AssetsToSync.Num())
 	{
diff --git a/Engine/Source/Editor/UnrealEd/Public/Dialogs/SMeshMergingDialog.h b/Engine/Source/Editor/UnrealEd/Public/Dialogs/SMeshMergingDialog.h
--- a/Engine/Source/Editor/UnrealEd/Public/Dialogs/SMeshMergingDialog.h
+++ b/Engine/Source/Editor/UnrealEd/Public/Dialogs/SMeshMergingDialog.h
@@ -41,6 +41,12 @@ private:
 	bool IsLightmapChannelEnabled() const;
 	void SetTargetLightMapChannel(TSharedPtr<FString> NewSelection, ESelectInfo::Type SelectInfo);
 	void SetTargetLightMapResolution(TSharedPtr<FString> NewSelection, ESelectInfo::Type SelectInfo);
+
+	/**  */
+	ECheckBoxState GetExportSpecificLODEnabled() const;
+	void SetExportSpecificLODEnabled(ECheckBoxState NewValue);
+	bool IsExportSpecificLODEnabled() const;
+	void SetExportSpecificLODIndex(TSharedPtr<FString> NewSelection, ESelectInfo::Type SelectInfo);
 		
 	/**  */
 	ECheckBoxState GetImportVertexColors() const;
@@ -101,6 +107,10 @@ private:
 	/** Whether to spawn merged actor in the world */
 	bool bPlaceInWorld;
 
+	bool bExportSpecificLOD;
+	int32 ExportLODIndex;
+	TArray<TSharedPtr<FString>>	ExportLODOptions;
+
 	/**  */
 	TArray<TSharedPtr<FString>>	LightMapResolutionOptions;
 	TArray<TSharedPtr<FString>>	LightMapChannelOptions;
