Fixed: Merge Actor tool overwrites multimaterial setup on static mesh LODs
UE-14650

diff --git a/Engine/Source/Developer/MeshUtilities/Private/MeshUtilities.cpp b/Engine/Source/Developer/MeshUtilities/Private/MeshUtilities.cpp
--- a/Engine/Source/Developer/MeshUtilities/Private/MeshUtilities.cpp
+++ b/Engine/Source/Developer/MeshUtilities/Private/MeshUtilities.cpp
@@ -3881,9 +3881,9 @@ static void ExportStaticMeshLOD(const FStaticMeshLODResources& StaticMeshLOD, FR
 		for (const FStaticMeshSection& Section : StaticMeshLOD.Sections)
 		{
 			uint32 FirstTriangle = Section.FirstIndex/3;
-			for (uint32 TriangleIndex = FirstTriangle; TriangleIndex < Section.NumTriangles; ++TriangleIndex)
+			for (uint32 TriangleIndex = 0; TriangleIndex < Section.NumTriangles; ++TriangleIndex)
 			{
-				OutRawMesh.FaceMaterialIndices[TriangleIndex] = Section.MaterialIndex;
+				OutRawMesh.FaceMaterialIndices[FirstTriangle + TriangleIndex] = Section.MaterialIndex;
 			}
 		}
 	}
@@ -4142,6 +4142,21 @@ static void CopyTextureRect(const FColor* Src, const FIntPoint& SrcSize, FColor*
 	}
 }
 
+static void SetTextureRect(const FColor& ColorValue, const FIntPoint& SrcSize, FColor* Dst, const FIntPoint& DstSize, const FIntPoint& DstPos)
+{
+	FColor* RowDst = Dst + DstSize.X*DstPos.Y;
+	
+	for (int32 RowIdx = 0; RowIdx < SrcSize.Y; ++RowIdx)
+	{
+		for (int32 ColIdx = 0; ColIdx < SrcSize.X; ++ColIdx)
+		{
+			RowDst[ColIdx] = ColorValue;
+		}
+				
+		RowDst+= DstSize.X;
+	}
+}
+
 struct FRawMeshExt
 {
 	FRawMeshExt() 
@@ -4275,33 +4290,54 @@ static void MergeMaterials(UWorld* InWorld, const TArray<UMaterialInterface*>& I
 		FlatMaterial.MetallicSize	= bExportMetallic ? ExportTextureSize : FIntPoint::ZeroValue;
 		FlatMaterial.RoughnessSize	= bExportRoughness ? ExportTextureSize : FIntPoint::ZeroValue;
 		FlatMaterial.SpecularSize	= bExportSpecular ? ExportTextureSize : FIntPoint::ZeroValue;
+		int32 ExportNumSamples		= ExportTextureSize.X*ExportTextureSize.Y;
 
 		FMaterialUtilities::ExportMaterial(InWorld, Material, FlatMaterial);
 
-		if (FlatMaterial.DiffuseSamples.Num() > 0)
+		if (FlatMaterial.DiffuseSamples.Num() == ExportNumSamples)
 		{
 			CopyTextureRect(FlatMaterial.DiffuseSamples.GetData(), ExportTextureSize, OutMergedMaterial.DiffuseSamples.GetData(), AtlasTextureSize, AtlasTargetPos);
 		}
+		else if (FlatMaterial.DiffuseSamples.Num() == 1)
+		{
+			SetTextureRect(FlatMaterial.DiffuseSamples[0], ExportTextureSize, OutMergedMaterial.DiffuseSamples.GetData(), AtlasTextureSize, AtlasTargetPos);
+		}
 
-		if (FlatMaterial.NormalSamples.Num() > 0)
+		if (FlatMaterial.NormalSamples.Num() == ExportNumSamples)
 		{
 			CopyTextureRect(FlatMaterial.NormalSamples.GetData(), ExportTextureSize, OutMergedMaterial.NormalSamples.GetData(), AtlasTextureSize, AtlasTargetPos);
 		}
-
-		if (FlatMaterial.MetallicSamples.Num() > 0)
+		else if (FlatMaterial.NormalSamples.Num() == 1)
+		{
+			SetTextureRect(FlatMaterial.NormalSamples[0], ExportTextureSize, OutMergedMaterial.NormalSamples.GetData(), AtlasTextureSize, AtlasTargetPos);
+		}
+		
+		if (FlatMaterial.MetallicSamples.Num() == ExportNumSamples)
 		{
 			CopyTextureRect(FlatMaterial.MetallicSamples.GetData(), ExportTextureSize, OutMergedMaterial.MetallicSamples.GetData(), AtlasTextureSize, AtlasTargetPos);
 		}
+		else if (FlatMaterial.MetallicSamples.Num() == 1)
+		{
+			SetTextureRect(FlatMaterial.MetallicSamples[0], ExportTextureSize, OutMergedMaterial.MetallicSamples.GetData(), AtlasTextureSize, AtlasTargetPos);
+		}
 
-		if (FlatMaterial.RoughnessSamples.Num() > 0)
+		if (FlatMaterial.RoughnessSamples.Num() == ExportNumSamples)
 		{
 			CopyTextureRect(FlatMaterial.RoughnessSamples.GetData(), ExportTextureSize, OutMergedMaterial.RoughnessSamples.GetData(), AtlasTextureSize, AtlasTargetPos);
 		}
+		else if (FlatMaterial.RoughnessSamples.Num() == 1)
+		{
+			SetTextureRect(FlatMaterial.RoughnessSamples[0], ExportTextureSize, OutMergedMaterial.RoughnessSamples.GetData(), AtlasTextureSize, AtlasTargetPos);
+		}
 
-		if (FlatMaterial.SpecularSamples.Num() > 0)
+		if (FlatMaterial.SpecularSamples.Num() == ExportNumSamples)
 		{
 			CopyTextureRect(FlatMaterial.SpecularSamples.GetData(), ExportTextureSize, OutMergedMaterial.SpecularSamples.GetData(), AtlasTextureSize, AtlasTargetPos);
 		}
+		else if (FlatMaterial.SpecularSamples.Num() == 1)
+		{
+			SetTextureRect(FlatMaterial.SpecularSamples[0], ExportTextureSize, OutMergedMaterial.SpecularSamples.GetData(), AtlasTextureSize, AtlasTargetPos);
+		}
 		
 		check(OutUVTransforms.IsValidIndex(MatIdx));
 
