Log Material list before merging

diff --git a/Engine/Source/Developer/MeshUtilities/Private/MeshUtilities.cpp b/Engine/Source/Developer/MeshUtilities/Private/MeshUtilities.cpp
--- a/Engine/Source/Developer/MeshUtilities/Private/MeshUtilities.cpp
+++ b/Engine/Source/Developer/MeshUtilities/Private/MeshUtilities.cpp
@@ -4210,6 +4210,24 @@ static void MergeMaterials(UWorld* InWorld, const TArray<UMaterialInterface*>& I
 {
 	OutUVTransforms.Reserve(InMaterialList.Num());
 
+#if !(UE_BUILD_SHIPPING || UE_BUILD_TEST)
+	// add log to warn which material it
+	UE_LOG(LogMeshUtilities,Log,TEXT("Merging Material: Merge Material count %d"),InMaterialList.Num());
+
+	int32 Index=0;
+	for (auto& MaterialIter : InMaterialList)
+	{
+		if (MaterialIter)
+		{
+			UE_LOG(LogMeshUtilities,Log,TEXT("Material List %d - %s"), ++Index, *MaterialIter->GetName());
+		}
+		else
+		{
+			UE_LOG(LogMeshUtilities,Log,TEXT("Material List %d - null material"), ++Index);
+		}
+	}
+
+#endif
 	// We support merging only for opaque materials
 	int32 NumOpaqueMaterials = 0;
 	// Fill output UV transforms with invalid values
