Fixed: HLOD does not generate non-overlapping lightmaps UE-14609

diff --git a/Engine/Source/Developer/MeshUtilities/Private/MeshUtilities.cpp b/Engine/Source/Developer/MeshUtilities/Private/MeshUtilities.cpp
--- a/Engine/Source/Developer/MeshUtilities/Private/MeshUtilities.cpp
+++ b/Engine/Source/Developer/MeshUtilities/Private/MeshUtilities.cpp
@@ -4713,6 +4713,11 @@ void FMeshUtilities::MergeActors(
 
 		// make sure it has a new lighting guid
 		StaticMesh->LightingGuid = FGuid::NewGuid();
+		if (InSettings.bGenerateLightMapUV)
+		{
+			StaticMesh->LightMapResolution = InSettings.TargetLightMapResolution;
+			StaticMesh->LightMapCoordinateIndex = InSettings.TargetLightMapUVChannel;
+		}
 		
 		for (int32 LODIndex = 0; LODIndex < NumMaxLOD; ++LODIndex)
 		{
