Fixed: Building HLOD assets warns that flag bUsedWithStaticLighting needs to be set
UE-15400

diff --git a/Engine/Source/Developer/MeshUtilities/Private/MeshUtilities.cpp b/Engine/Source/Developer/MeshUtilities/Private/MeshUtilities.cpp
--- a/Engine/Source/Developer/MeshUtilities/Private/MeshUtilities.cpp
+++ b/Engine/Source/Developer/MeshUtilities/Private/MeshUtilities.cpp
@@ -3771,6 +3771,15 @@ void FMeshUtilities::CreateProxyMesh(
 	// Construct proxy material
 	UMaterial* ProxyMaterial = FMaterialUtilities::CreateMaterial(ProxyFlattenMaterial, InOuter, ProxyBasePackageName, RF_Public|RF_Standalone, OutAssetsToSync);
 	
+	// Set material static lighting usage flag if project has static lighting enabled
+	static const auto AllowStaticLightingVar = IConsoleManager::Get().FindTConsoleVariableDataInt(TEXT("r.AllowStaticLighting"));
+	const bool bAllowStaticLighting = (!AllowStaticLightingVar || AllowStaticLightingVar->GetValueOnGameThread() != 0);
+	if (bAllowStaticLighting)
+	{
+		bool bNeedsRecompile;
+		ProxyMaterial->SetMaterialUsage(bNeedsRecompile, MATUSAGE_StaticLighting);
+	}
+		
 	// Construct proxy static mesh
 	UPackage* MeshPackage = InOuter;
 	FString MeshAssetName = TEXT("SM_") + AssetBaseName;
@@ -4598,6 +4607,16 @@ void FMeshUtilities::MergeActors(
 		}
 
 		UMaterial* MergedMaterial = FMaterialUtilities::CreateMaterial(MergedFlatMaterial, MaterialPackage, MaterialAssetName, RF_Public|RF_Standalone, OutAssetsToSync);
+
+		// Set material static lighting usage flag if project has static lighting enabled
+		static const auto AllowStaticLightingVar = IConsoleManager::Get().FindTConsoleVariableDataInt(TEXT("r.AllowStaticLighting"));
+		const bool bAllowStaticLighting = (!AllowStaticLightingVar || AllowStaticLightingVar->GetValueOnGameThread() != 0);
+		if (bAllowStaticLighting)
+		{
+			bool bNeedsRecompile;
+			MergedMaterial->SetMaterialUsage(bNeedsRecompile, MATUSAGE_StaticLighting);
+		}
+
 		UniqueMaterials.Add(MergedMaterial);
 	}
 		
