Make sure ActiveBoneIndices is sorted on root lod when importing skeletal meshes

diff --git a/Engine/Source/Developer/MeshUtilities/Private/MeshUtilities.cpp b/Engine/Source/Developer/MeshUtilities/Private/MeshUtilities.cpp
--- a/Engine/Source/Developer/MeshUtilities/Private/MeshUtilities.cpp
+++ b/Engine/Source/Developer/MeshUtilities/Private/MeshUtilities.cpp
@@ -994,6 +994,8 @@ void FMeshUtilities::BuildSkeletalModelFromChunks(FStaticLODModel& LODModel,cons
 			LODModel.ActiveBoneIndices.AddUnique(Chunk.BoneMap[BoneIndex]);
 		}
 	}
+	
+	LODModel.ActiveBoneIndices.Sort();
 
 	// Reset 'final vertex to import vertex' map info
 	LODModel.MeshToImportVertexMap.Empty();
