Bumped the max mesh material count warning value

diff --git a/Engine/Source/Developer/MeshUtilities/Public/MeshUtilities.h b/Engine/Source/Developer/MeshUtilities/Public/MeshUtilities.h
--- a/Engine/Source/Developer/MeshUtilities/Public/MeshUtilities.h
+++ b/Engine/Source/Developer/MeshUtilities/Public/MeshUtilities.h
@@ -5,7 +5,7 @@
 #include "ModuleInterface.h"
 
 /** The maximum index for a material in a mesh */
-#define MAX_MESH_MATERIAL_INDEX 64
+#define MAX_MESH_MATERIAL_INDEX 256
 
 /**
  * Mesh reduction interface.
