Fixed: HLOD ignores Build Scale when using Simplygon
UE-16327

diff --git a/Engine/Source/Developer/MeshUtilities/Private/MeshUtilities.cpp b/Engine/Source/Developer/MeshUtilities/Private/MeshUtilities.cpp
--- a/Engine/Source/Developer/MeshUtilities/Private/MeshUtilities.cpp
+++ b/Engine/Source/Developer/MeshUtilities/Private/MeshUtilities.cpp
@@ -3969,8 +3969,18 @@ bool FMeshUtilities::ConstructRawMesh(
 		}
 	}
 
+	// Use build settings from base mesh for LOD entries that was generated inside Editor.
+	const FMeshBuildSettings& BuildSettings = bImportedMesh ? SrcModel.BuildSettings : SrcMesh->SourceModels[0].BuildSettings;
+
 	// Transform raw mesh to world space
 	FTransform CtoM = InMeshComponent->ComponentToWorld;
+	// Take into account build scale settings only for meshes imported from raw data
+	// meshes reconstructed from render data already have build scale applied
+	if (bImportedMesh)
+	{
+		CtoM.SetScale3D(CtoM.GetScale3D()*BuildSettings.BuildScale3D);
+	}
+
 	for (FVector& Vertex : OutRawMesh.VertexPositions)
 	{
 		Vertex = CtoM.TransformPosition(Vertex);
@@ -4016,9 +4026,6 @@ bool FMeshUtilities::ConstructRawMesh(
 		
 	int32 NumWedges = OutRawMesh.WedgeIndices.Num();
 
-	// Use build settings from base mesh for LOD entries that was generated inside Editor.
-	const FMeshBuildSettings& BuildSettings = bImportedMesh ? SrcModel.BuildSettings : SrcMesh->SourceModels[0].BuildSettings;
-
 	// Figure out if we should recompute normals and tangents. By default generated LODs should not recompute normals
 	bool bRecomputeNormals = (bImportedMesh && BuildSettings.bRecomputeNormals) || OutRawMesh.WedgeTangentZ.Num() == 0;
 	bool bRecomputeTangents = (bImportedMesh && BuildSettings.bRecomputeTangents) || OutRawMesh.WedgeTangentX.Num() == 0 || OutRawMesh.WedgeTangentY.Num() == 0;
