Use previously imported/calculated tangents as initial value during skeletal fbx import smoothing step (better matches static behaviour). UE-16898

diff --git a/Engine/Source/Developer/MeshUtilities/Private/MeshUtilities.cpp b/Engine/Source/Developer/MeshUtilities/Private/MeshUtilities.cpp
--- a/Engine/Source/Developer/MeshUtilities/Private/MeshUtilities.cpp
+++ b/Engine/Source/Developer/MeshUtilities/Private/MeshUtilities.cpp
@@ -3371,18 +3371,19 @@ bool FMeshUtilities::BuildSkeletalMesh( FStaticLODModel& LODModel, const FRefere
 
 		if( bComputeNormals || bComputeTangents )
 		{
-			for(int32 VertexIndex = 0;VertexIndex < 3;VertexIndex++)
-			{
-				VertexTangentX[VertexIndex] = FVector::ZeroVector;
-				VertexTangentY[VertexIndex] = FVector::ZeroVector;
-				VertexTangentZ[VertexIndex] = FVector::ZeroVector;
-			}
-
 			FVector	TriangleNormal = FPlane(
 				Points[Wedges[Face.iWedge[2]].iVertex],
 				Points[Wedges[Face.iWedge[1]].iVertex],
 				Points[Wedges[Face.iWedge[0]].iVertex]
 			);
+
+			for(int32 VertexIndex = 0;VertexIndex < 3;VertexIndex++)
+			{
+				VertexTangentX[VertexIndex] = FaceTangentX[FaceIndex];
+				VertexTangentY[VertexIndex] = FaceTangentY[FaceIndex];
+				VertexTangentZ[VertexIndex] = TriangleNormal;
+			}
+
 			float	Determinant = FVector::Triple(FaceTangentX[FaceIndex],FaceTangentY[FaceIndex],TriangleNormal);
 
 			// Start building a list of faces adjacent to this triangle
