Increase limit on maximum size of index buffer sections to cache optimize

--------
Integrated using branch UE4-To-UE4-Orion (reversed) of change#2548128 by [REDACTED_PERSON] on 2015/05/12 18:35:46.

diff --git a/Engine/Source/Developer/MeshUtilities/Private/MeshUtilities.cpp b/Engine/Source/Developer/MeshUtilities/Private/MeshUtilities.cpp
--- a/Engine/Source/Developer/MeshUtilities/Private/MeshUtilities.cpp
+++ b/Engine/Source/Developer/MeshUtilities/Private/MeshUtilities.cpp
@@ -31,7 +31,7 @@
 // causes meshes to be rebuilt you MUST generate a new GUID and replace this
 // string with it.
 
-#define MESH_UTILITIES_VER TEXT("2C1BC8F50A7A43818AFE266EB43D9060")
+#define MESH_UTILITIES_VER TEXT("8C68575CEF434CA8A9E1DA4AED8A47BB")
 
 DEFINE_LOG_CATEGORY_STATIC(LogMeshUtilities,Verbose,All);
 
@@ -701,11 +701,11 @@ namespace Forsyth
 }
 
 void FMeshUtilities::CacheOptimizeIndexBuffer(TArray<uint16>& Indices)
-{
+{	
 	if(bUsingNvTriStrip)
-	{
-		NvTriStrip::CacheOptimizeIndexBuffer(Indices);
-	}
+	{				
+			NvTriStrip::CacheOptimizeIndexBuffer(Indices);
+		}
 	else if( !bDisableTriangleOrderOptimization )
 	{
 		Forsyth::CacheOptimizeIndexBuffer(Indices);
@@ -716,10 +716,10 @@ void FMeshUtilities::CacheOptimizeIndexBuffer(TArray<uint32>& Indices)
 {
 	if(bUsingNvTriStrip)
 	{
-		NvTriStrip::CacheOptimizeIndexBuffer(Indices);
-	}
+			NvTriStrip::CacheOptimizeIndexBuffer(Indices);
+		}
 	else if( !bDisableTriangleOrderOptimization )
-	{
+	{		
 		Forsyth::CacheOptimizeIndexBuffer(Indices);
 	}
 }
@@ -2951,7 +2951,7 @@ public:
 				BuildStaticMeshVertexAndIndexBuffers(Vertices, PerSectionIndices, WedgeMap, RawMesh, LODOverlappingCorners[LODIndex], ComparisonThreshold, LODBuildSettings[LODIndex].BuildScale3D);
 				check(WedgeMap.Num() == RawMesh.WedgeIndices.Num());
 
-				if (RawMesh.WedgeIndices.Num() < 50000 * 3)
+				if (RawMesh.WedgeIndices.Num() < 100000 * 3)
 				{
 					MeshUtilities.CacheOptimizeVertexAndIndexBuffer(Vertices, PerSectionIndices, WedgeMap);
 					check(WedgeMap.Num() == RawMesh.WedgeIndices.Num());
