Mesh Merging: Lightmap channel index will always follow last vallid UV channel in case it does not override existing channel
Partial fix for UE-15319

diff --git a/Engine/Source/Developer/MeshUtilities/Private/MeshUtilities.cpp b/Engine/Source/Developer/MeshUtilities/Private/MeshUtilities.cpp
--- a/Engine/Source/Developer/MeshUtilities/Private/MeshUtilities.cpp
+++ b/Engine/Source/Developer/MeshUtilities/Private/MeshUtilities.cpp
@@ -4762,6 +4762,16 @@ void FMeshUtilities::MergeActors(
 			}
 		}
 	}
+
+	// Compute target lightmap channel for each LOD
+	// User can specify any index, but there are should not be empty gaps in UV channel list
+	int32 TargetLightMapUVChannel[MAX_STATIC_MESH_LODS];
+	for (int32 LODIndex = 0; LODIndex < MAX_STATIC_MESH_LODS; ++LODIndex)
+	{
+		int32 ChannelIdx = 0;
+		for (; ChannelIdx < MAX_MESH_TEXTURE_COORDS && bOcuppiedUVChannels[LODIndex][ChannelIdx]; ++ChannelIdx);
+		TargetLightMapUVChannel[LODIndex] = FMath::Min(InSettings.TargetLightMapUVChannel, ChannelIdx);
+	}
 	
 	//
 	//Create merged mesh asset
@@ -4799,7 +4809,7 @@ void FMeshUtilities::MergeActors(
 		if (InSettings.bGenerateLightMapUV)
 		{
 			StaticMesh->LightMapResolution = InSettings.TargetLightMapResolution;
-			StaticMesh->LightMapCoordinateIndex = InSettings.TargetLightMapUVChannel;
+			StaticMesh->LightMapCoordinateIndex = TargetLightMapUVChannel[0];
 		}
 		
 		for (int32 LODIndex = 0; LODIndex < NumMaxLOD; ++LODIndex)
@@ -4816,7 +4826,7 @@ void FMeshUtilities::MergeActors(
 				SrcModel->BuildSettings.bGenerateLightmapUVs = InSettings.bGenerateLightMapUV;
 				SrcModel->BuildSettings.MinLightmapResolution = InSettings.TargetLightMapResolution;
 				SrcModel->BuildSettings.SrcLightmapIndex = 0;
-				SrcModel->BuildSettings.DstLightmapIndex = InSettings.TargetLightMapUVChannel;
+				SrcModel->BuildSettings.DstLightmapIndex = TargetLightMapUVChannel[LODIndex];
 
 				SrcModel->RawMeshBulkData->SaveRawMesh(MergedMeshLOD);
 			}
