Added ability to merge physics data (collision primitives) for Actor Merging tool
UE-14879

diff --git a/Engine/Source/Developer/MeshUtilities/Private/MeshUtilities.cpp b/Engine/Source/Developer/MeshUtilities/Private/MeshUtilities.cpp
--- a/Engine/Source/Developer/MeshUtilities/Private/MeshUtilities.cpp
+++ b/Engine/Source/Developer/MeshUtilities/Private/MeshUtilities.cpp
@@ -16,6 +16,7 @@
 #include "DistanceFieldAtlas.h"
 #include "FbxErrors.h"
 #include "Components/SplineMeshComponent.h"
+#include "PhysicsEngine/BodySetup.h"
 #include "MaterialUtilities.h"
 
 //@todo - implement required vector intrinsics for other implementations
@@ -4210,7 +4211,7 @@ void FMeshUtilities::CalculateRawMeshTangents(FRawMesh& OutRawMesh, const FMeshB
 /*------------------------------------------------------------------------------
 	Mesh merging 
 ------------------------------------------------------------------------------*/
-bool PropagatePaintedColorsToRawMesh(UStaticMeshComponent* StaticMeshComponent, int32 LODIndex, FRawMesh& RawMesh)
+static bool PropagatePaintedColorsToRawMesh(UStaticMeshComponent* StaticMeshComponent, int32 LODIndex, FRawMesh& RawMesh)
 {
 	UStaticMesh* StaticMesh = StaticMeshComponent->StaticMesh;
 	
@@ -4259,6 +4260,83 @@ bool PropagatePaintedColorsToRawMesh(UStaticMeshComponent* StaticMeshComponent,
 	return false;
 }
 
+static void TransformPhysicsGeometry(const FTransform& InTransform, FKAggregateGeom& AggGeom)
+{
+	for (auto& Elem : AggGeom.SphereElems)
+	{
+		FTransform ElemTM = Elem.GetTransform();
+		Elem.SetTransform(ElemTM*InTransform);
+	}
+	
+	for (auto& Elem : AggGeom.BoxElems)
+	{
+		FTransform ElemTM = Elem.GetTransform();
+		Elem.SetTransform(ElemTM*InTransform);
+	}
+
+	for (auto& Elem : AggGeom.SphylElems)
+	{
+		FTransform ElemTM = Elem.GetTransform();
+		Elem.SetTransform(ElemTM*InTransform);
+	}
+
+	for (auto& Elem : AggGeom.ConvexElems)
+	{
+		FTransform ElemTM = Elem.GetTransform();
+		Elem.SetTransform(ElemTM*InTransform);
+	}
+
+	// seems like all primitives except Convex need separate scaling pass		
+	const FVector Scale3D = InTransform.GetScale3D();
+	if (!Scale3D.Equals(FVector(1.f)))
+	{
+		const float MinPrimSize = KINDA_SMALL_NUMBER;
+		
+		for (auto& Elem : AggGeom.SphereElems)
+		{
+			Elem.ScaleElem(Scale3D, MinPrimSize);
+		}
+	
+		for (auto& Elem : AggGeom.BoxElems)
+		{
+			Elem.ScaleElem(Scale3D, MinPrimSize);
+		}
+
+		for (auto& Elem : AggGeom.SphylElems)
+		{
+			Elem.ScaleElem(Scale3D, MinPrimSize);
+		}
+	}
+}
+
+static void ExtractPhysicsGeometry(UStaticMeshComponent* InMeshComponent, FKAggregateGeom& OutAggGeom)
+{
+	UStaticMesh* SrcMesh = InMeshComponent->StaticMesh;
+	if (SrcMesh == nullptr)
+	{
+		return;
+	}
+	
+	if (!SrcMesh->BodySetup)
+	{
+		return;
+	}
+
+	OutAggGeom = SrcMesh->BodySetup->AggGeom;
+	
+	// we are not owner of this stuff
+	OutAggGeom.RenderInfo = nullptr; 
+	for (auto& Elem : OutAggGeom.ConvexElems)
+	{
+		Elem.ConvexMesh = nullptr;
+		Elem.ConvexMeshNegX = nullptr;
+	}
+
+	// Transform geometry to world space
+	FTransform CtoM = InMeshComponent->ComponentToWorld;
+	TransformPhysicsGeometry(CtoM, OutAggGeom);
+}
+
 static void CopyTextureRect(const FColor* Src, const FIntPoint& SrcSize, FColor* Dst, const FIntPoint& DstSize, const FIntPoint& DstPos)
 {
 	int32 RowLength = SrcSize.X*sizeof(FColor);
@@ -4291,12 +4369,10 @@ static void SetTextureRect(const FColor& ColorValue, const FIntPoint& SrcSize, F
 
 struct FRawMeshExt
 {
-	FRawMeshExt() 
-	{}
-		
-	FRawMesh	MeshLOD[MAX_STATIC_MESH_LODS];
-	FString		AssetPackageName;
-	FVector		Pivot;	
+	FRawMesh		MeshLOD[MAX_STATIC_MESH_LODS];
+	FString			AssetPackageName;
+	FVector			Pivot;
+	FKAggregateGeom	AggGeom;
 };
 
 struct FRawMeshUVTransform
@@ -4509,6 +4585,7 @@ static void MergeMaterials(UWorld* InWorld, const TArray<UMaterialInterface*>& I
 		AtlasTargetPos = FIntPoint(AtlasColIdx*ExportTextureSize.X, AtlasRowIdx*ExportTextureSize.Y);
 	}
 }
+
 void FMeshUtilities::MergeActors(
 	const TArray<AActor*>& SourceActors,
 	const FMeshMergingSettings& InSettings,
@@ -4540,8 +4617,9 @@ void FMeshUtilities::MergeActors(
 	}
 
 	UWorld* World = SourceActors[0]->GetWorld();
+	float ViewDistance = TNumericLimits<float>::Max();
 
-	MergeStaticMeshComponents(ComponentsToMerge, World, InSettings, InOuter, InBasePackageName, UseLOD, OutAssetsToSync, OutMergedActorLocation, bSilent);
+	MergeStaticMeshComponents(ComponentsToMerge, World, InSettings, InOuter, InBasePackageName, UseLOD, OutAssetsToSync, OutMergedActorLocation, ViewDistance, bSilent);
 }
 
 void FMeshUtilities::MergeStaticMeshComponents(const TArray<UStaticMeshComponent*>& ComponentsToMerge, UWorld* World, const FMeshMergingSettings& InSettings, UPackage* InOuter, const FString& InBasePackageName, int32 UseLOD, /* does not build all LODs but only use this LOD to create base mesh */ TArray<UObject*>& OutAssetsToSync, FVector& OutMergedActorLocation, float ViewDistance, bool bSilent /*= false*/) const
@@ -4553,7 +4631,13 @@ void FMeshUtilities::MergeStaticMeshComponents(const TArray<UStaticMeshComponent
 	TArray<FRawMeshExt>								SourceMeshes;
 	bool											bWithVertexColors[MAX_STATIC_MESH_LODS] = {};
 	bool											bOcuppiedUVChannels[MAX_STATIC_MESH_LODS][MAX_MESH_TEXTURE_COORDS] = {};
+	UBodySetup*										BodySetupSource = nullptr;
 	FScopedSlowTask									MergeProgress(5);
+
+	if (ComponentsToMerge.Num() == 0)
+	{
+		return;
+	}
 	
 	SourceMeshes.Reserve(ComponentsToMerge.Num());
 	SourceMeshes.SetNum(ComponentsToMerge.Num());
@@ -4569,9 +4653,6 @@ void FMeshUtilities::MergeStaticMeshComponents(const TArray<UStaticMeshComponent
 		SourceMeshes[MeshId].Pivot = MeshComponent->ComponentToWorld.GetLocation();
 		// Source mesh asset package name
 		SourceMeshes[MeshId].AssetPackageName = MeshComponent->StaticMesh->GetOutermost()->GetName();
-
-
-		bool check = true;
 	}
 
 	NumMaxLOD = FMath::Min(NumMaxLOD, MAX_STATIC_MESH_LODS);
@@ -4621,11 +4702,21 @@ void FMeshUtilities::MergeStaticMeshComponents(const TArray<UStaticMeshComponent
 		}
 	}
 
-	if (SourceMeshes.Num() == 0)
+	if (InSettings.bMergePhysicsData)
 	{
-		return;
+		for (int32 MeshId = 0; MeshId < ComponentsToMerge.Num(); ++MeshId)
+		{
+			UStaticMeshComponent* MeshComponent = ComponentsToMerge[MeshId];
+			ExtractPhysicsGeometry(MeshComponent, SourceMeshes[MeshId].AggGeom);
+			
+			// We will use first valid BodySetup as a source of physics settings
+			if (BodySetupSource == nullptr)
+			{
+				BodySetupSource = MeshComponent->StaticMesh->BodySetup;
+			}
+		}
 	}
-
+	
 	MergeProgress.EnterProgressFrame();
 
 	// For each raw mesh, re-map the material indices according to the MaterialMap
@@ -4840,6 +4931,16 @@ void FMeshUtilities::MergeStaticMeshComponents(const TArray<UStaticMeshComponent
 		}
 	}
 
+	// Transform physics primitives to merged mesh pivot
+	if (InSettings.bMergePhysicsData && !MergedMesh.Pivot.IsZero())
+	{
+		FTransform PivotTM(-MergedMesh.Pivot);
+		for (auto& SourceMesh : SourceMeshes)
+		{
+			TransformPhysicsGeometry(PivotTM, SourceMesh.AggGeom);
+		}
+	}
+
 	// Compute target lightmap channel for each LOD
 	// User can specify any index, but there are should not be empty gaps in UV channel list
 	int32 TargetLightMapUVChannel[MAX_STATIC_MESH_LODS];
@@ -4922,6 +5023,18 @@ void FMeshUtilities::MergeStaticMeshComponents(const TArray<UStaticMeshComponent
 			StaticMesh->Materials.Add(Material);
 		}
 
+		if (InSettings.bMergePhysicsData && BodySetupSource)
+		{
+			StaticMesh->CreateBodySetup();
+			StaticMesh->BodySetup->CopyBodyPropertiesFrom(BodySetupSource);
+			StaticMesh->BodySetup->AggGeom = FKAggregateGeom();
+			
+			for (const FRawMeshExt& SourceMesh : SourceMeshes)
+			{
+				StaticMesh->BodySetup->AddCollisionFrom(SourceMesh.AggGeom);
+			}
+		}
+
 		StaticMesh->Build(bSilent);
 		StaticMesh->PostEditChange();
 
diff --git a/Engine/Source/Editor/MergeActors/Private/MeshMergingTool/SMeshMergingDialog.cpp b/Engine/Source/Editor/MergeActors/Private/MeshMergingTool/SMeshMergingDialog.cpp
--- a/Engine/Source/Editor/MergeActors/Private/MeshMergingTool/SMeshMergingDialog.cpp
+++ b/Engine/Source/Editor/MergeActors/Private/MeshMergingTool/SMeshMergingDialog.cpp
@@ -231,6 +231,23 @@ void SMeshMergingDialog::Construct(const FArguments& InArgs, FMeshMergingTool* I
 						.Font(FEditorStyle::GetFontStyle("StandardDialog.SmallFont"))
 					]
 				]
+				
+				// Merge physics data
+				+SVerticalBox::Slot()
+				.AutoHeight()
+				.Padding(FEditorStyle::GetMargin("StandardDialog.ContentPadding"))
+				[
+					SNew(SCheckBox)
+					.Type(ESlateCheckBoxType::CheckBox)
+					.IsChecked(this, &SMeshMergingDialog::GetMergePhyisicData)
+					.OnCheckStateChanged(this, &SMeshMergingDialog::SetMergePhyisicData)
+					.Content()
+					[
+						SNew(STextBlock)
+						.Text(LOCTEXT("MergePhysicsDataLabel", "Merge Physics Data"))
+						.Font(FEditorStyle::GetFontStyle("StandardDialog.SmallFont"))
+					]
+				]
 
 				// Merge materials
 				+SVerticalBox::Slot()
@@ -427,6 +444,16 @@ void SMeshMergingDialog::SetReplaceSourceActors(ECheckBoxState NewValue)
 	Tool->bReplaceSourceActors = (ECheckBoxState::Checked == NewValue);
 }
 
+ECheckBoxState SMeshMergingDialog::GetMergePhyisicData() const
+{
+	return (Tool->MergingSettings.bMergePhysicsData ? ECheckBoxState::Checked : ECheckBoxState::Unchecked);
+}
+
+void SMeshMergingDialog::SetMergePhyisicData(ECheckBoxState NewValue)
+{
+	Tool->MergingSettings.bMergePhysicsData = (ECheckBoxState::Checked == NewValue);
+}
+
 bool SMeshMergingDialog::IsMaterialMergingEnabled() const
 {
 	return Tool->MergingSettings.bMergeMaterials;
diff --git a/Engine/Source/Editor/MergeActors/Private/MeshMergingTool/SMeshMergingDialog.h b/Engine/Source/Editor/MergeActors/Private/MeshMergingTool/SMeshMergingDialog.h
--- a/Engine/Source/Editor/MergeActors/Private/MeshMergingTool/SMeshMergingDialog.h
+++ b/Engine/Source/Editor/MergeActors/Private/MeshMergingTool/SMeshMergingDialog.h
@@ -53,6 +53,9 @@ private:
 	ECheckBoxState GetReplaceSourceActors() const;
 	void SetReplaceSourceActors(ECheckBoxState NewValue);
 
+	ECheckBoxState GetMergePhyisicData() const;
+	void SetMergePhyisicData(ECheckBoxState NewValue);
+	
 	/** Material merging */
 	bool IsMaterialMergingEnabled() const;
 	ECheckBoxState GetMergeMaterials() const;
diff --git a/Engine/Source/Runtime/Engine/Classes/Engine/EngineTypes.h b/Engine/Source/Runtime/Engine/Classes/Engine/EngineTypes.h
--- a/Engine/Source/Runtime/Engine/Classes/Engine/EngineTypes.h
+++ b/Engine/Source/Runtime/Engine/Classes/Engine/EngineTypes.h
@@ -2457,6 +2457,10 @@ struct FMeshMergingSettings
 	UPROPERTY(EditAnywhere, Category=FMeshMergingSettings)
 	bool bPivotPointAtZero;
 
+	/** Whether to merge physics data (collision primitives)*/
+	UPROPERTY(EditAnywhere, Category=FMeshMergingSettings)
+	bool bMergePhysicsData;
+
 	/** Whether to merge source materials into one flat material */
 	UPROPERTY(EditAnywhere, Category=MeshMerge)
 	bool bMergeMaterials;
@@ -2483,6 +2487,7 @@ struct FMeshMergingSettings
 		, TargetLightMapResolution(256)
 		, bImportVertexColors(false)
 		, bPivotPointAtZero(false)
+		, bMergePhysicsData(false)
 		, bMergeMaterials(false)
 		, bExportNormalMap(true)
 		, bExportMetallicMap(false)
diff --git a/Engine/Source/Runtime/Engine/Classes/PhysicsEngine/BodySetup.h b/Engine/Source/Runtime/Engine/Classes/PhysicsEngine/BodySetup.h
--- a/Engine/Source/Runtime/Engine/Classes/PhysicsEngine/BodySetup.h
+++ b/Engine/Source/Runtime/Engine/Classes/PhysicsEngine/BodySetup.h
@@ -152,8 +152,8 @@ public:
 
 	/** Add collision shapes from another body setup to this one */
 	ENGINE_API void AddCollisionFrom(class UBodySetup* FromSetup);
-
-
+	ENGINE_API void AddCollisionFrom(const FKAggregateGeom& FromAggGeom);
+	
 	/** Create Physics meshes (ConvexMeshes, TriMesh & TriMeshNegX) from cooked data */
 	/** Release Physics meshes (ConvexMeshes, TriMesh & TriMeshNegX). Must be called before the BodySetup is destroyed */
 	ENGINE_API virtual void CreatePhysicsMeshes();
diff --git a/Engine/Source/Runtime/Engine/Private/PhysicsEngine/BodySetup.cpp b/Engine/Source/Runtime/Engine/Private/PhysicsEngine/BodySetup.cpp
--- a/Engine/Source/Runtime/Engine/Private/PhysicsEngine/BodySetup.cpp
+++ b/Engine/Source/Runtime/Engine/Private/PhysicsEngine/BodySetup.cpp
@@ -73,17 +73,17 @@ void UBodySetup::CopyBodyPropertiesFrom(const UBodySetup* FromSetup)
 	bDoubleSidedGeometry = FromSetup->bDoubleSidedGeometry;
 }
 
-void UBodySetup::AddCollisionFrom(class UBodySetup* FromSetup)
+void UBodySetup::AddCollisionFrom(const FKAggregateGeom& FromAggGeom)
 {
 	// Add shapes from static mesh
-	AggGeom.SphereElems.Append(FromSetup->AggGeom.SphereElems);
-	AggGeom.BoxElems.Append(FromSetup->AggGeom.BoxElems);
-	AggGeom.SphylElems.Append(FromSetup->AggGeom.SphylElems);
+	AggGeom.SphereElems.Append(FromAggGeom.SphereElems);
+	AggGeom.BoxElems.Append(FromAggGeom.BoxElems);
+	AggGeom.SphylElems.Append(FromAggGeom.SphylElems);
 
 	// Remember how many convex we already have
 	int32 FirstNewConvexIdx = AggGeom.ConvexElems.Num();
 	// copy convex
-	AggGeom.ConvexElems.Append(FromSetup->AggGeom.ConvexElems);
+	AggGeom.ConvexElems.Append(FromAggGeom.ConvexElems);
 	// clear pointers on convex elements
 	for (int32 i = FirstNewConvexIdx; i < AggGeom.ConvexElems.Num(); i++)
 	{
@@ -93,6 +93,11 @@ void UBodySetup::AddCollisionFrom(class UBodySetup* FromSetup)
 	}
 }
 
+void UBodySetup::AddCollisionFrom(class UBodySetup* FromSetup)
+{
+	AddCollisionFrom(FromSetup->AggGeom);
+}
+
 DECLARE_CYCLE_STAT(TEXT("Create Physics Meshes"), STAT_CreatePhysicsMeshes, STATGROUP_Physics);
 
 
