Merging back (experimental) tangent import work and fixes: 2649527, 2649599, 2652924, 2660845.

diff --git a/Engine/Config/BaseEditorPerProjectUserSettings.ini b/Engine/Config/BaseEditorPerProjectUserSettings.ini
--- a/Engine/Config/BaseEditorPerProjectUserSettings.ini
+++ b/Engine/Config/BaseEditorPerProjectUserSettings.ini
@@ -589,6 +589,7 @@ VertexOverrideColor=(R=255,G=255,B=255,A=255)
 bPreserveSmoothingGroups=True
 bImportMeshLODs=False
 bImportMorphTargets=False
+bUseExperimentalTangentGeneration=False
 
 
 [/Script/UnrealEd.FbxTextureImportData]
diff --git a/Engine/Source/Developer/MeshUtilities/Private/MeshUtilities.cpp b/Engine/Source/Developer/MeshUtilities/Private/MeshUtilities.cpp
--- a/Engine/Source/Developer/MeshUtilities/Private/MeshUtilities.cpp
+++ b/Engine/Source/Developer/MeshUtilities/Private/MeshUtilities.cpp
@@ -96,7 +96,8 @@ private:
 		bool bGenerateAsIfTwoSided,
 		FDistanceFieldVolumeData& OutData) override;
 
-	virtual bool BuildSkeletalMesh( FStaticLODModel& LODModel, const FReferenceSkeleton& RefSkeleton, const TArray<FVertInfluence>& Influences, const TArray<FMeshWedge>& Wedges, const TArray<FMeshFace>& Faces, const TArray<FVector>& Points, const TArray<int32>& PointToOriginalMap, bool bKeepOverlappingVertices = false, bool bComputeNormals = true, bool bComputeTangents = true, TArray<FText> * OutWarningMessages = NULL, TArray<FName> * OutWarningNames = NULL) override;
+	virtual bool BuildSkeletalMesh( FStaticLODModel& LODModel, const FReferenceSkeleton& RefSkeleton, const TArray<FVertInfluence>& Influences, const TArray<FMeshWedge>& Wedges, const TArray<FMeshFace>& Faces, const TArray<FVector>& Points, const TArray<int32>& PointToOriginalMap, const MeshBuildOptions& BuildOptions = MeshBuildOptions(), TArray<FText> * OutWarningMessages = NULL, TArray<FName> * OutWarningNames = NULL) override;
+	bool BuildSkeletalMesh_Legacy( FStaticLODModel& LODModel, const FReferenceSkeleton& RefSkeleton, const TArray<FVertInfluence>& Influences, const TArray<FMeshWedge>& Wedges, const TArray<FMeshFace>& Faces, const TArray<FVector>& Points, const TArray<int32>& PointToOriginalMap, bool bKeepOverlappingVertices = false, bool bComputeNormals = true, bool bComputeTangents = true, TArray<FText> * OutWarningMessages = NULL, TArray<FName> * OutWarningNames = NULL);
 
 	virtual IMeshReduction* GetMeshReductionInterface() override;
 	virtual IMeshMerging* GetMeshMergingInterface() override;
@@ -2018,7 +2019,7 @@ public:
 	const TArray<FMeshWedge>	&wedges;			//Reference to wedge list.
 	const TArray<FMeshFace>		&faces;				//Reference to face list.	Also contains normal/tangent/bitanget/UV coords for each vertex of the face.
 	const TArray<FVector>		&points;			//Reference to position list.
-	bool						&bComputeNormals;	//Reference to bComputeNormals.
+	bool						bComputeNormals;	//Copy of bComputeNormals.
 	TArray<FVector>				&TangentsX;			//Reference to newly created tangents list.
 	TArray<FVector>				&TangentsY;			//Reference to newly created bitangents list.
 	TArray<FVector>				&TangentsZ;			//Reference to computed normals, will be empty otherwise.
@@ -2027,7 +2028,7 @@ public:
 		const TArray<FMeshWedge>	&Wedges,
 		const TArray<FMeshFace>		&Faces,
 		const TArray<FVector>		&Points,
-		bool						&bInComputeNormals,
+		bool						bInComputeNormals,
 		TArray<FVector>				&VertexTangentsX,
 		TArray<FVector>				&VertexTangentsY,
 		TArray<FVector>				&VertexTangentsZ
@@ -3191,10 +3192,1270 @@ bool FMeshUtilities::GenerateStaticMeshLODs(TArray<FStaticMeshSourceModel>& Mode
 	return false;
 }
 
-//@TODO: The OutMessages has to be a struct that contains FText/FName, or make it Token and add that as error. Needs re-work. Temporary workaround for now. 
-bool FMeshUtilities::BuildSkeletalMesh( FStaticLODModel& LODModel, const FReferenceSkeleton& RefSkeleton, const TArray<FVertInfluence>& Influences, const TArray<FMeshWedge>& Wedges, const TArray<FMeshFace>& Faces, const TArray<FVector>& Points, const TArray<int32>& PointToOriginalMap, bool bKeepOverlappingVertices, bool bComputeNormals, bool bComputeTangents, TArray<FText> * OutWarningMessages, TArray<FName> * OutWarningNames)
+class IMeshBuildData
+{
+public:
+	virtual uint32 GetWedgeIndex(uint32 FaceIndex, uint32 TriIndex) = 0;
+	virtual uint32 GetVertexIndex(uint32 WedgeIndex) = 0;
+	virtual uint32 GetVertexIndex(uint32 FaceIndex, uint32 TriIndex) = 0;
+	virtual FVector GetVertexPosition(uint32 WedgeIndex) = 0;
+	virtual FVector GetVertexPosition(uint32 FaceIndex, uint32 TriIndex) = 0;
+	virtual FVector2D GetVertexUV(uint32 FaceIndex, uint32 TriIndex, uint32 UVIndex) = 0;
+
+	virtual uint32 GetNumFaces() = 0;
+	virtual uint32 GetNumWedges() = 0;
+
+	virtual TArray<FVector>& GetTangentArray(uint32 Axis) = 0;
+	virtual void ValidateTangentArraySize() = 0;
+
+	virtual SMikkTSpaceInterface* GetMikkTInterface() = 0;
+	virtual void* GetMikkTUserData() = 0;
+
+	const IMeshUtilities::MeshBuildOptions& BuildOptions;
+	TArray<FText>* OutWarningMessages;
+	TArray<FName>* OutWarningNames;
+	bool bTooManyVerts;
+
+protected:
+	IMeshBuildData(
+		const IMeshUtilities::MeshBuildOptions& InBuildOptions,
+		TArray<FText>* InWarningMessages,
+		TArray<FName>* InWarningNames)
+	: BuildOptions(InBuildOptions)
+	, OutWarningMessages(InWarningMessages)
+	, OutWarningNames(InWarningNames)
+	, bTooManyVerts(false)
+	{
+	}
+};
+
+class SkeletalMeshBuildData : public IMeshBuildData
+{
+public:
+	SkeletalMeshBuildData(
+		FStaticLODModel& InLODModel,
+		const FReferenceSkeleton& InRefSkeleton,
+		const TArray<FVertInfluence>& InInfluences,
+		const TArray<FMeshWedge>& InWedges,
+		const TArray<FMeshFace>& InFaces,
+		const TArray<FVector>& InPoints,
+		const TArray<int32>& InPointToOriginalMap,
+		const IMeshUtilities::MeshBuildOptions& InBuildOptions,
+		TArray<FText>* InWarningMessages,
+		TArray<FName>* InWarningNames)
+	: IMeshBuildData(InBuildOptions, InWarningMessages, InWarningNames)
+	, LODModel(InLODModel)
+	, RefSkeleton(InRefSkeleton)
+	, Influences(InInfluences)
+	, Wedges(InWedges)
+	, Faces(InFaces)
+	, Points(InPoints)
+	, PointToOriginalMap(InPointToOriginalMap)
+	, MikkTUserData(InWedges, InFaces, InPoints, InBuildOptions.bComputeNormals, TangentX, TangentY, TangentZ)
+	{
+		MikkTInterface.m_getNormal = MikkGetNormal_Skeletal;
+		MikkTInterface.m_getNumFaces = MikkGetNumFaces_Skeletal;
+		MikkTInterface.m_getNumVerticesOfFace = MikkGetNumVertsOfFace_Skeletal;
+		MikkTInterface.m_getPosition = MikkGetPosition_Skeletal;
+		MikkTInterface.m_getTexCoord = MikkGetTexCoord_Skeletal;
+		MikkTInterface.m_setTSpaceBasic = MikkSetTSpaceBasic_Skeletal;
+		MikkTInterface.m_setTSpace = nullptr;
+	}
+
+	virtual uint32 GetWedgeIndex(uint32 FaceIndex, uint32 TriIndex) override
+	{
+		return Faces[FaceIndex].iWedge[TriIndex];
+	}
+
+	virtual uint32 GetVertexIndex(uint32 WedgeIndex) override
+	{
+		return Wedges[WedgeIndex].iVertex;
+	}
+
+	virtual uint32 GetVertexIndex(uint32 FaceIndex, uint32 TriIndex) override
+	{	
+		return Wedges[Faces[FaceIndex].iWedge[TriIndex]].iVertex;
+	}
+
+	virtual FVector GetVertexPosition(uint32 WedgeIndex) override
+	{
+		return Points[Wedges[WedgeIndex].iVertex];
+	}
+
+	virtual FVector GetVertexPosition(uint32 FaceIndex, uint32 TriIndex) override
+	{
+		return Points[Wedges[Faces[FaceIndex].iWedge[TriIndex]].iVertex];
+	}
+
+	virtual FVector2D GetVertexUV(uint32 FaceIndex, uint32 TriIndex, uint32 UVIndex) override
+	{
+		return Wedges[Faces[FaceIndex].iWedge[TriIndex]].UVs[UVIndex];
+	}
+
+	virtual uint32 GetNumFaces() override
+	{
+		return Faces.Num();
+	}
+
+	virtual uint32 GetNumWedges() override
+	{
+		return Wedges.Num();
+	}
+
+	virtual TArray<FVector>& GetTangentArray(uint32 Axis) override
+	{
+		if (Axis == 0)
+		{
+			return TangentX;
+		}
+		else if (Axis == 1)
+		{
+			return TangentY;
+		}
+
+		return TangentZ;
+	}
+
+	virtual void ValidateTangentArraySize() override
+	{
+		check(TangentX.Num() == Wedges.Num());
+		check(TangentY.Num() == Wedges.Num());
+		check(TangentZ.Num() == Wedges.Num());
+	}
+
+	virtual SMikkTSpaceInterface* GetMikkTInterface() override
+	{
+		return &MikkTInterface;
+	}
+
+	virtual void* GetMikkTUserData() override
+	{
+		return (void*)&MikkTUserData;
+	}
+
+	TArray<FVector> TangentX;
+	TArray<FVector> TangentY;
+	TArray<FVector> TangentZ;
+	TArray<FSkinnedMeshChunk*> Chunks;
+
+	SMikkTSpaceInterface MikkTInterface;
+	MikkTSpace_Skeletal_Mesh MikkTUserData;
+
+	FStaticLODModel& LODModel;
+	const FReferenceSkeleton& RefSkeleton;
+	const TArray<FVertInfluence>& Influences;
+	const TArray<FMeshWedge>& Wedges;
+	const TArray<FMeshFace>& Faces;
+	const TArray<FVector>& Points;
+	const TArray<int32>& PointToOriginalMap;
+};
+
+class FSkeletalMeshUtilityBuilder
+{
+public:
+	FSkeletalMeshUtilityBuilder()
+	: Stage(EStage::Uninit)
+	{
+	}
+
+public:
+	void Skeletal_FindOverlappingCorners(
+		TMultiMap<int32,int32>& OutOverlappingCorners,
+		IMeshBuildData* BuildData,
+		float ComparisonThreshold
+		)
+	{
+		int32 NumFaces = BuildData->GetNumFaces();
+		int32 NumWedges = BuildData->GetNumWedges();
+		check(NumFaces * 3 <= NumWedges);
+
+		// Create a list of vertex Z/index pairs
+		TArray<FIndexAndZ> VertIndexAndZ;
+		VertIndexAndZ.Empty(NumWedges);
+		for(int32 FaceIndex = 0; FaceIndex < NumFaces; FaceIndex++)
+		{
+			for (int32 TriIndex = 0; TriIndex < 3; ++TriIndex)
+			{
+				uint32 Index = BuildData->GetWedgeIndex(FaceIndex, TriIndex);
+				new(VertIndexAndZ) FIndexAndZ(Index, BuildData->GetVertexPosition(Index));
+			}	
+		}
+
+		// Sort the vertices by z value
+		VertIndexAndZ.Sort(FCompareIndexAndZ());
+
+		// Search for duplicates, quickly!
+		for (int32 i = 0; i < VertIndexAndZ.Num(); i++)
+		{
+			// only need to search forward, since we add pairs both ways
+			for (int32 j = i + 1; j < VertIndexAndZ.Num(); j++)
+			{
+				if (FMath::Abs(VertIndexAndZ[j].Z - VertIndexAndZ[i].Z) > ComparisonThreshold)
+					break; // can't be any more dups
+
+				FVector PositionA = BuildData->GetVertexPosition(VertIndexAndZ[i].Index);
+				FVector PositionB = BuildData->GetVertexPosition(VertIndexAndZ[j].Index);
+
+				if (PointsEqual(PositionA, PositionB, ComparisonThreshold))					
+				{
+					OutOverlappingCorners.Add(VertIndexAndZ[i].Index,VertIndexAndZ[j].Index);
+					OutOverlappingCorners.Add(VertIndexAndZ[j].Index,VertIndexAndZ[i].Index);
+				}
+			}
+		}
+	}
+
+	void Skeletal_ComputeTriangleTangents(
+		TArray<FVector>& TriangleTangentX,
+		TArray<FVector>& TriangleTangentY,
+		TArray<FVector>& TriangleTangentZ,
+		IMeshBuildData* BuildData,
+		float ComparisonThreshold
+		)
+	{
+		int32 NumTriangles = BuildData->GetNumFaces();
+		TriangleTangentX.Empty(NumTriangles);
+		TriangleTangentY.Empty(NumTriangles);
+		TriangleTangentZ.Empty(NumTriangles);
+
+		for (int32 TriangleIndex = 0; TriangleIndex < NumTriangles; TriangleIndex++)
+		{
+			const int32 UVIndex = 0;
+			FVector P[3];
+
+			for (int32 i = 0; i < 3; ++i)
+			{
+				P[i] = BuildData->GetVertexPosition(TriangleIndex, i);
+			}
+
+			const FVector Normal = ((P[1] - P[2])^(P[0] - P[2])).GetSafeNormal(ComparisonThreshold);
+			FMatrix	ParameterToLocal(
+				FPlane(P[1].X - P[0].X, P[1].Y - P[0].Y, P[1].Z - P[0].Z, 0),
+				FPlane(P[2].X - P[0].X, P[2].Y - P[0].Y, P[2].Z - P[0].Z, 0),
+				FPlane(P[0].X,          P[0].Y,          P[0].Z,          0),
+				FPlane(0,               0,               0,				  1)
+				);
+
+			FVector2D T1 = BuildData->GetVertexUV(TriangleIndex, 0, UVIndex);
+			FVector2D T2 = BuildData->GetVertexUV(TriangleIndex, 1, UVIndex);
+			FVector2D T3 = BuildData->GetVertexUV(TriangleIndex, 2, UVIndex);
+			FMatrix ParameterToTexture(
+				FPlane(	T2.X - T1.X,	T2.Y - T1.Y,	0,	0	),
+				FPlane(	T3.X - T1.X,	T3.Y - T1.Y,	0,	0	),
+				FPlane(	T1.X,			T1.Y,			1,	0	),
+				FPlane(	0,				0,				0,	1	)
+				);
+
+			// Use InverseSlow to catch singular matrices.  Inverse can miss this sometimes.
+			const FMatrix TextureToLocal = ParameterToTexture.Inverse() * ParameterToLocal;
+
+			TriangleTangentX.Add(TextureToLocal.TransformVector(FVector(1,0,0)).GetSafeNormal());
+			TriangleTangentY.Add(TextureToLocal.TransformVector(FVector(0,1,0)).GetSafeNormal());
+			TriangleTangentZ.Add(Normal);
+
+			FVector::CreateOrthonormalBasis(
+				TriangleTangentX[TriangleIndex],
+				TriangleTangentY[TriangleIndex],
+				TriangleTangentZ[TriangleIndex]
+				);
+		}
+	}
+
+	void Skeletal_ComputeTangents(
+		IMeshBuildData* BuildData,
+		TMultiMap<int32,int32> const& OverlappingCorners
+		)
+	{
+		bool bBlendOverlappingNormals = true;
+		bool bIgnoreDegenerateTriangles = BuildData->BuildOptions.bRemoveDegenerateTriangles;
+		float ComparisonThreshold = bIgnoreDegenerateTriangles ? THRESH_POINTS_ARE_SAME : 0.0f;
+
+		// Compute per-triangle tangents.
+		TArray<FVector> TriangleTangentX;
+		TArray<FVector> TriangleTangentY;
+		TArray<FVector> TriangleTangentZ;
+
+		Skeletal_ComputeTriangleTangents(
+			TriangleTangentX,
+			TriangleTangentY,
+			TriangleTangentZ,
+			BuildData,
+			bIgnoreDegenerateTriangles ? SMALL_NUMBER : 0.0f
+			);
+
+		TArray<FVector>& WedgeTangentX = BuildData->GetTangentArray(0);
+		TArray<FVector>& WedgeTangentY = BuildData->GetTangentArray(1);
+		TArray<FVector>& WedgeTangentZ = BuildData->GetTangentArray(2);
+
+		// Declare these out here to avoid reallocations.
+		TArray<FFanFace> RelevantFacesForCorner[3];
+		TArray<int32> AdjacentFaces;
+		TArray<int32> DupVerts;
+
+		int32 NumFaces = BuildData->GetNumFaces();
+		int32 NumWedges = BuildData->GetNumWedges();
+		check(NumFaces * 3 <= NumWedges);
+
+		// Allocate storage for tangents if none were provided.
+		if (WedgeTangentX.Num() != NumWedges)
+		{
+			WedgeTangentX.Empty(NumWedges);
+			WedgeTangentX.AddZeroed(NumWedges);
+		}
+		if (WedgeTangentY.Num() != NumWedges)
+		{
+			WedgeTangentY.Empty(NumWedges);
+			WedgeTangentY.AddZeroed(NumWedges);
+		}
+		if (WedgeTangentZ.Num() != NumWedges)
+		{
+			WedgeTangentZ.Empty(NumWedges);
+			WedgeTangentZ.AddZeroed(NumWedges);
+		}
+
+		for (int32 FaceIndex = 0; FaceIndex < NumFaces; FaceIndex++)
+		{
+			int32 WedgeOffset = FaceIndex * 3;
+			FVector CornerPositions[3];
+			FVector CornerTangentX[3];
+			FVector CornerTangentY[3];
+			FVector CornerTangentZ[3];
+
+			for (int32 CornerIndex = 0; CornerIndex < 3; CornerIndex++)
+			{
+				CornerTangentX[CornerIndex] = FVector::ZeroVector;
+				CornerTangentY[CornerIndex] = FVector::ZeroVector;
+				CornerTangentZ[CornerIndex] = FVector::ZeroVector;
+				CornerPositions[CornerIndex] = BuildData->GetVertexPosition(FaceIndex, CornerIndex);
+				RelevantFacesForCorner[CornerIndex].Reset();
+			}
+
+			// Don't process degenerate triangles.
+			if (PointsEqual(CornerPositions[0],CornerPositions[1], ComparisonThreshold)
+				|| PointsEqual(CornerPositions[0],CornerPositions[2], ComparisonThreshold)
+				|| PointsEqual(CornerPositions[1],CornerPositions[2], ComparisonThreshold))
+			{
+				continue;
+			}
+
+			// No need to process triangles if tangents already exist.
+			bool bCornerHasTangents[3] = {0};
+			for (int32 CornerIndex = 0; CornerIndex < 3; CornerIndex++)
+			{
+				bCornerHasTangents[CornerIndex] = !WedgeTangentX[WedgeOffset + CornerIndex].IsZero()
+					&& !WedgeTangentY[WedgeOffset + CornerIndex].IsZero()
+					&& !WedgeTangentZ[WedgeOffset + CornerIndex].IsZero();
+			}
+			if (bCornerHasTangents[0] && bCornerHasTangents[1] && bCornerHasTangents[2])
+			{
+				continue;
+			}
+
+			// Calculate smooth vertex normals.
+			float Determinant = FVector::Triple(
+				TriangleTangentX[FaceIndex],
+				TriangleTangentY[FaceIndex],
+				TriangleTangentZ[FaceIndex]
+				);
+
+			// Start building a list of faces adjacent to this face.
+			AdjacentFaces.Reset();
+			for (int32 CornerIndex = 0; CornerIndex < 3; CornerIndex++)
+			{
+				int32 ThisCornerIndex = WedgeOffset + CornerIndex;
+				DupVerts.Reset();
+				OverlappingCorners.MultiFind(ThisCornerIndex,DupVerts);
+				DupVerts.Add(ThisCornerIndex); // I am a "dup" of myself
+				for (int32 k = 0; k < DupVerts.Num(); k++)
+				{
+					AdjacentFaces.AddUnique(DupVerts[k] / 3);
+				}
+			}
+
+			// We need to sort these here because the criteria for point equality is
+			// exact, so we must ensure the exact same order for all dups.
+			AdjacentFaces.Sort();
+
+			// Process adjacent faces
+			for (int32 AdjacentFaceIndex = 0; AdjacentFaceIndex < AdjacentFaces.Num(); AdjacentFaceIndex++)
+			{
+				int32 OtherFaceIndex = AdjacentFaces[AdjacentFaceIndex];
+				for (int32 OurCornerIndex = 0; OurCornerIndex < 3; OurCornerIndex++)
+				{
+					if (bCornerHasTangents[OurCornerIndex])
+						continue;
+
+					FFanFace NewFanFace;
+					int32 CommonIndexCount = 0;
+
+					// Check for vertices in common.
+					if (FaceIndex == OtherFaceIndex)
+					{
+						CommonIndexCount = 3;		
+						NewFanFace.LinkedVertexIndex = OurCornerIndex;
+					}
+					else
+					{
+						// Check matching vertices against main vertex .
+						for (int32 OtherCornerIndex = 0; OtherCornerIndex < 3; OtherCornerIndex++)
+						{
+							if (PointsEqual(
+									CornerPositions[OurCornerIndex],
+									BuildData->GetVertexPosition(OtherFaceIndex, OtherCornerIndex),
+									ComparisonThreshold
+									))
+							{
+								CommonIndexCount++;
+								NewFanFace.LinkedVertexIndex = OtherCornerIndex;
+							}
+						}
+					}
+
+					// Add if connected by at least one point. Smoothing matches are considered later.
+					if (CommonIndexCount > 0)
+					{ 					
+						NewFanFace.FaceIndex = OtherFaceIndex;
+						NewFanFace.bFilled = (OtherFaceIndex == FaceIndex); // Starter face for smoothing floodfill.
+						NewFanFace.bBlendTangents = NewFanFace.bFilled;
+						NewFanFace.bBlendNormals = NewFanFace.bFilled;
+						RelevantFacesForCorner[OurCornerIndex].Add(NewFanFace);
+					}
+				}
+			}
+
+			// Find true relevance of faces for a vertex normal by traversing
+			// smoothing-group-compatible connected triangle fans around common vertices.
+			for (int32 CornerIndex = 0; CornerIndex < 3; CornerIndex++)
+			{
+				if (bCornerHasTangents[CornerIndex])
+					continue;
+
+				int32 NewConnections;
+				do
+				{
+					NewConnections = 0;
+					for (int32 OtherFaceIdx=0; OtherFaceIdx < RelevantFacesForCorner[CornerIndex].Num(); OtherFaceIdx++)
+					{
+						FFanFace& OtherFace = RelevantFacesForCorner[CornerIndex][OtherFaceIdx];
+						// The vertex' own face is initially the only face with bFilled == true.
+						if (OtherFace.bFilled)
+						{				
+							for (int32 NextFaceIndex = 0; NextFaceIndex < RelevantFacesForCorner[CornerIndex].Num(); NextFaceIndex++)
+							{
+								FFanFace& NextFace = RelevantFacesForCorner[CornerIndex][NextFaceIndex];
+								if (!NextFace.bFilled) // && !NextFace.bBlendTangents)
+								{
+									if (NextFaceIndex != OtherFaceIdx)
+										//&& (RawMesh.FaceSmoothingMasks[NextFace.FaceIndex] & RawMesh.FaceSmoothingMasks[OtherFace.FaceIndex]))
+									{				
+										int32 CommonVertices = 0;
+										int32 CommonTangentVertices = 0;
+										int32 CommonNormalVertices = 0;
+										for (int32 OtherCornerIndex = 0; OtherCornerIndex < 3; OtherCornerIndex++)
+										{											
+											for (int32 NextCornerIndex = 0; NextCornerIndex < 3; NextCornerIndex++)
+											{
+												int32 NextVertexIndex = BuildData->GetVertexIndex(NextFace.FaceIndex, NextCornerIndex);
+												int32 OtherVertexIndex = BuildData->GetVertexIndex(OtherFace.FaceIndex, OtherCornerIndex);
+												if (PointsEqual(
+													BuildData->GetVertexPosition(NextFace.FaceIndex, NextCornerIndex),
+													BuildData->GetVertexPosition(OtherFace.FaceIndex, OtherCornerIndex),
+													ComparisonThreshold))
+												{
+													CommonVertices++;
+
+													
+													if (UVsEqual(
+															BuildData->GetVertexUV(NextFace.FaceIndex, NextCornerIndex, 0),
+															BuildData->GetVertexUV(OtherFace.FaceIndex, OtherCornerIndex, 0)))
+													{
+														CommonTangentVertices++;
+													}
+													if (bBlendOverlappingNormals
+														|| NextVertexIndex == OtherVertexIndex)
+													{
+														CommonNormalVertices++;
+													}
+												}
+											}										
+										}
+										// Flood fill faces with more than one common vertices which must be touching edges.
+										if (CommonVertices > 1)
+										{
+											NextFace.bFilled = true;
+											NextFace.bBlendNormals = (CommonNormalVertices > 1);
+											NewConnections++;
+
+											// Only blend tangents if there is no UV seam along the edge with this face.
+											if (OtherFace.bBlendTangents && CommonTangentVertices > 1)
+											{
+												float OtherDeterminant = FVector::Triple(
+													TriangleTangentX[NextFace.FaceIndex],
+													TriangleTangentY[NextFace.FaceIndex],
+													TriangleTangentZ[NextFace.FaceIndex]
+													);
+												if ((Determinant * OtherDeterminant) > 0.0f)
+												{
+													NextFace.bBlendTangents = true;
+												}
+											}
+										}								
+									}
+								}
+							}
+						}
+					}
+				}
+				while (NewConnections > 0);
+			}
+
+			// Vertex normal construction.
+			for (int32 CornerIndex = 0; CornerIndex < 3; CornerIndex++)
+			{
+				if (bCornerHasTangents[CornerIndex])
+				{
+					CornerTangentX[CornerIndex] = WedgeTangentX[WedgeOffset + CornerIndex];
+					CornerTangentY[CornerIndex] = WedgeTangentY[WedgeOffset + CornerIndex];
+					CornerTangentZ[CornerIndex] = WedgeTangentZ[WedgeOffset + CornerIndex];
+				}
+				else
+				{
+					for (int32 RelevantFaceIdx = 0; RelevantFaceIdx < RelevantFacesForCorner[CornerIndex].Num(); RelevantFaceIdx++)
+					{
+						FFanFace const& RelevantFace = RelevantFacesForCorner[CornerIndex][RelevantFaceIdx];
+						if (RelevantFace.bFilled)
+						{
+							int32 OtherFaceIndex = RelevantFace.FaceIndex;
+							if (RelevantFace.bBlendTangents)
+							{
+								CornerTangentX[CornerIndex] += TriangleTangentX[OtherFaceIndex];
+								CornerTangentY[CornerIndex] += TriangleTangentY[OtherFaceIndex];
+							}
+							if (RelevantFace.bBlendNormals)
+							{
+								CornerTangentZ[CornerIndex] += TriangleTangentZ[OtherFaceIndex];
+							}
+						}
+					}
+					if (!WedgeTangentX[WedgeOffset + CornerIndex].IsZero())
+					{
+						CornerTangentX[CornerIndex] = WedgeTangentX[WedgeOffset + CornerIndex];
+					}
+					if (!WedgeTangentY[WedgeOffset + CornerIndex].IsZero())
+					{
+						CornerTangentY[CornerIndex] = WedgeTangentY[WedgeOffset + CornerIndex];
+					}
+					if (!WedgeTangentZ[WedgeOffset + CornerIndex].IsZero())
+					{
+						CornerTangentZ[CornerIndex] = WedgeTangentZ[WedgeOffset + CornerIndex];
+					}
+				}
+			}
+
+			// Normalization.
+			for (int32 CornerIndex = 0; CornerIndex < 3; CornerIndex++)
+			{
+				CornerTangentX[CornerIndex].Normalize();
+				CornerTangentY[CornerIndex].Normalize();
+				CornerTangentZ[CornerIndex].Normalize();
+
+				// Gram-Schmidt orthogonalization
+				CornerTangentY[CornerIndex] -= CornerTangentX[CornerIndex] * (CornerTangentX[CornerIndex] | CornerTangentY[CornerIndex]);
+				CornerTangentY[CornerIndex].Normalize();
+
+				CornerTangentX[CornerIndex] -= CornerTangentZ[CornerIndex] * (CornerTangentZ[CornerIndex] | CornerTangentX[CornerIndex]);
+				CornerTangentX[CornerIndex].Normalize();
+				CornerTangentY[CornerIndex] -= CornerTangentZ[CornerIndex] * (CornerTangentZ[CornerIndex] | CornerTangentY[CornerIndex]);
+				CornerTangentY[CornerIndex].Normalize();
+			}
+
+			// Copy back to the mesh.
+			for (int32 CornerIndex = 0; CornerIndex < 3; CornerIndex++)
+			{
+				WedgeTangentX[WedgeOffset + CornerIndex] = CornerTangentX[CornerIndex];
+				WedgeTangentY[WedgeOffset + CornerIndex] = CornerTangentY[CornerIndex];
+				WedgeTangentZ[WedgeOffset + CornerIndex] = CornerTangentZ[CornerIndex];
+			}
+		}
+
+		check(WedgeTangentX.Num() == NumWedges);
+		check(WedgeTangentY.Num() == NumWedges);
+		check(WedgeTangentZ.Num() == NumWedges);
+	}
+
+	void Skeletal_ComputeTangents_MikkTSpace(
+		IMeshBuildData* BuildData,
+		TMultiMap<int32,int32> const& OverlappingCorners
+		)
+	{
+		bool bBlendOverlappingNormals = true;
+		bool bIgnoreDegenerateTriangles = BuildData->BuildOptions.bRemoveDegenerateTriangles;
+		float ComparisonThreshold = bIgnoreDegenerateTriangles ? THRESH_POINTS_ARE_SAME : 0.0f;
+
+		// Compute per-triangle tangents.
+		TArray<FVector> TriangleTangentX;
+		TArray<FVector> TriangleTangentY;
+		TArray<FVector> TriangleTangentZ;
+
+		Skeletal_ComputeTriangleTangents(
+			TriangleTangentX,
+			TriangleTangentY,
+			TriangleTangentZ,
+			BuildData,
+			bIgnoreDegenerateTriangles ? SMALL_NUMBER : 0.0f
+			);
+
+		TArray<FVector>& WedgeTangentX = BuildData->GetTangentArray(0);
+		TArray<FVector>& WedgeTangentY = BuildData->GetTangentArray(1);
+		TArray<FVector>& WedgeTangentZ = BuildData->GetTangentArray(2);
+
+		// Declare these out here to avoid reallocations.
+		TArray<FFanFace> RelevantFacesForCorner[3];
+		TArray<int32> AdjacentFaces;
+		TArray<int32> DupVerts;
+
+		int32 NumFaces = BuildData->GetNumFaces();
+		int32 NumWedges = BuildData->GetNumWedges();
+		check(NumFaces * 3 == NumWedges);
+
+		bool bWedgeNormals = true;
+		bool bWedgeTSpace = false;
+		for (int32 WedgeIdx = 0; WedgeIdx < WedgeTangentZ.Num(); ++WedgeIdx)
+		{
+			for(int32 VertexIndex = 0;VertexIndex < 3;VertexIndex++)
+			bWedgeNormals = bWedgeNormals && (!WedgeTangentZ[WedgeIdx].IsNearlyZero());
+		}
+
+		if (WedgeTangentX.Num() > 0 && WedgeTangentY.Num() > 0)
+		{
+			bWedgeTSpace = true;
+			for (int32 WedgeIdx = 0; WedgeIdx < WedgeTangentX.Num()
+				&& WedgeIdx < WedgeTangentY.Num(); ++WedgeIdx)
+			{
+				bWedgeTSpace = bWedgeTSpace && (!WedgeTangentX[WedgeIdx].IsNearlyZero()) && (!WedgeTangentY[WedgeIdx].IsNearlyZero());
+			}
+		}
+
+		// Allocate storage for tangents if none were provided, and calculate normals for MikkTSpace.
+		if (WedgeTangentZ.Num() != NumWedges || !bWedgeNormals)
+		{
+			// normals are not included, so we should calculate them
+			WedgeTangentZ.Empty(NumWedges);
+			WedgeTangentZ.AddZeroed(NumWedges);
+			// we need to calculate normals for MikkTSpace
+
+			for (int32 FaceIndex = 0; FaceIndex < NumFaces; FaceIndex++)
+			{
+				int32 WedgeOffset = FaceIndex * 3;
+				FVector CornerPositions[3];
+				FVector CornerNormal[3];
+
+				for (int32 CornerIndex = 0; CornerIndex < 3; CornerIndex++)
+				{
+					CornerNormal[CornerIndex] = FVector::ZeroVector;
+					CornerPositions[CornerIndex] = BuildData->GetVertexPosition(FaceIndex, CornerIndex);
+					RelevantFacesForCorner[CornerIndex].Reset();
+				}
+
+				// Don't process degenerate triangles.
+				if (PointsEqual(CornerPositions[0], CornerPositions[1], ComparisonThreshold)
+					|| PointsEqual(CornerPositions[0], CornerPositions[2], ComparisonThreshold)
+					|| PointsEqual(CornerPositions[1], CornerPositions[2], ComparisonThreshold))
+				{
+					continue;
+				}
+
+				// No need to process triangles if tangents already exist.
+				bool bCornerHasNormal[3] = { 0 };
+				for (int32 CornerIndex = 0; CornerIndex < 3; CornerIndex++)
+				{
+					bCornerHasNormal[CornerIndex] = !WedgeTangentZ[WedgeOffset + CornerIndex].IsZero();
+				}
+				if (bCornerHasNormal[0] && bCornerHasNormal[1] && bCornerHasNormal[2])
+				{
+					continue;
+				}
+
+				// Start building a list of faces adjacent to this face.
+				AdjacentFaces.Reset();
+				for (int32 CornerIndex = 0; CornerIndex < 3; CornerIndex++)
+				{
+					int32 ThisCornerIndex = WedgeOffset + CornerIndex;
+					DupVerts.Reset();
+					OverlappingCorners.MultiFind(ThisCornerIndex,DupVerts);
+					DupVerts.Add(ThisCornerIndex); // I am a "dup" of myself
+					for (int32 k = 0; k < DupVerts.Num(); k++)
+					{
+						AdjacentFaces.AddUnique(DupVerts[k] / 3);
+					}
+				}
+
+				// We need to sort these here because the criteria for point equality is
+				// exact, so we must ensure the exact same order for all dups.
+				AdjacentFaces.Sort();
+
+				// Process adjacent faces
+				for (int32 AdjacentFaceIndex = 0; AdjacentFaceIndex < AdjacentFaces.Num(); AdjacentFaceIndex++)
+				{
+					int32 OtherFaceIndex = AdjacentFaces[AdjacentFaceIndex];
+					for (int32 OurCornerIndex = 0; OurCornerIndex < 3; OurCornerIndex++)
+					{
+						if (bCornerHasNormal[OurCornerIndex])
+							continue;
+
+						FFanFace NewFanFace;
+						int32 CommonIndexCount = 0;
+
+						// Check for vertices in common.
+						if (FaceIndex == OtherFaceIndex)
+						{
+							CommonIndexCount = 3;
+							NewFanFace.LinkedVertexIndex = OurCornerIndex;
+						}
+						else
+						{
+							// Check matching vertices against main vertex .
+							for (int32 OtherCornerIndex = 0; OtherCornerIndex < 3; OtherCornerIndex++)
+							{
+								if (PointsEqual(
+									CornerPositions[OurCornerIndex],
+									BuildData->GetVertexPosition(OtherFaceIndex, OtherCornerIndex),
+									ComparisonThreshold
+									))
+								{
+									CommonIndexCount++;
+									NewFanFace.LinkedVertexIndex = OtherCornerIndex;
+								}
+							}
+						}
+
+						// Add if connected by at least one point. Smoothing matches are considered later.
+						if (CommonIndexCount > 0)
+						{
+							NewFanFace.FaceIndex = OtherFaceIndex;
+							NewFanFace.bFilled = (OtherFaceIndex == FaceIndex); // Starter face for smoothing floodfill.
+							NewFanFace.bBlendTangents = NewFanFace.bFilled;
+							NewFanFace.bBlendNormals = NewFanFace.bFilled;
+							RelevantFacesForCorner[OurCornerIndex].Add(NewFanFace);
+						}
+					}
+				}
+
+				// Find true relevance of faces for a vertex normal by traversing
+				// smoothing-group-compatible connected triangle fans around common vertices.
+				for (int32 CornerIndex = 0; CornerIndex < 3; CornerIndex++)
+				{
+					if (bCornerHasNormal[CornerIndex])
+						continue;
+
+					int32 NewConnections;
+					do
+					{
+						NewConnections = 0;
+						for (int32 OtherFaceIdx = 0; OtherFaceIdx < RelevantFacesForCorner[CornerIndex].Num(); OtherFaceIdx++)
+						{
+							FFanFace& OtherFace = RelevantFacesForCorner[CornerIndex][OtherFaceIdx];
+							// The vertex' own face is initially the only face with bFilled == true.
+							if (OtherFace.bFilled)
+							{
+								for (int32 NextFaceIndex = 0; NextFaceIndex < RelevantFacesForCorner[CornerIndex].Num(); NextFaceIndex++)
+								{
+									FFanFace& NextFace = RelevantFacesForCorner[CornerIndex][NextFaceIndex];
+									if (!NextFace.bFilled) // && !NextFace.bBlendTangents)
+									{
+										if ((NextFaceIndex != OtherFaceIdx)
+											)//&& (RawMesh.FaceSmoothingMasks[NextFace.FaceIndex] & RawMesh.FaceSmoothingMasks[OtherFace.FaceIndex]))
+										{
+											int32 CommonVertices = 0;
+											int32 CommonNormalVertices = 0;
+											for (int32 OtherCornerIndex = 0; OtherCornerIndex < 3; OtherCornerIndex++)
+											{
+												for (int32 NextCornerIndex = 0; NextCornerIndex < 3; NextCornerIndex++)
+												{
+													int32 NextVertexIndex = BuildData->GetVertexIndex(NextFace.FaceIndex, NextCornerIndex);
+													int32 OtherVertexIndex = BuildData->GetVertexIndex(OtherFace.FaceIndex, OtherCornerIndex);
+													if (PointsEqual(
+														BuildData->GetVertexPosition(NextFace.FaceIndex, NextCornerIndex),
+														BuildData->GetVertexPosition(OtherFace.FaceIndex, OtherCornerIndex),
+														ComparisonThreshold))
+													{
+														CommonVertices++;
+														if (bBlendOverlappingNormals
+															|| NextVertexIndex == OtherVertexIndex)
+														{
+															CommonNormalVertices++;
+														}
+													}
+												}
+											}
+											// Flood fill faces with more than one common vertices which must be touching edges.
+											if (CommonVertices > 1)
+											{
+												NextFace.bFilled = true;
+												NextFace.bBlendNormals = (CommonNormalVertices > 1);
+												NewConnections++;
+											}
+										}
+									}
+								}
+							}
+						}
+					}
+					while (NewConnections > 0);
+				}
+
+				// Vertex normal construction.
+				for (int32 CornerIndex = 0; CornerIndex < 3; CornerIndex++)
+				{
+					if (bCornerHasNormal[CornerIndex])
+					{
+						CornerNormal[CornerIndex] = WedgeTangentZ[WedgeOffset + CornerIndex];
+					}
+					else
+					{
+						for (int32 RelevantFaceIdx = 0; RelevantFaceIdx < RelevantFacesForCorner[CornerIndex].Num(); RelevantFaceIdx++)
+						{
+							FFanFace const& RelevantFace = RelevantFacesForCorner[CornerIndex][RelevantFaceIdx];
+							if (RelevantFace.bFilled)
+							{
+								int32 OtherFaceIndex = RelevantFace.FaceIndex;
+								if (RelevantFace.bBlendNormals)
+								{
+									CornerNormal[CornerIndex] += TriangleTangentZ[OtherFaceIndex];
+								}
+							}
+						}
+						if (!WedgeTangentZ[WedgeOffset + CornerIndex].IsZero())
+						{
+							CornerNormal[CornerIndex] = WedgeTangentZ[WedgeOffset + CornerIndex];
+						}
+					}
+				}
+
+				// Normalization.
+				for (int32 CornerIndex = 0; CornerIndex < 3; CornerIndex++)
+				{
+					CornerNormal[CornerIndex].Normalize();
+				}
+
+				// Copy back to the mesh.
+				for (int32 CornerIndex = 0; CornerIndex < 3; CornerIndex++)
+				{
+					WedgeTangentZ[WedgeOffset + CornerIndex] = CornerNormal[CornerIndex];
+				}
+			}
+		}
+
+		if (WedgeTangentX.Num() != NumWedges)
+		{
+			WedgeTangentX.Empty(NumWedges);
+			WedgeTangentX.AddZeroed(NumWedges);
+		}
+		if (WedgeTangentY.Num() != NumWedges)
+		{
+			WedgeTangentY.Empty(NumWedges);
+			WedgeTangentY.AddZeroed(NumWedges);
+		}
+
+		//if (!bWedgeTSpace)
+		{
+			// we can use mikktspace to calculate the tangents		
+			SMikkTSpaceContext MikkTContext;
+			MikkTContext.m_pInterface = BuildData->GetMikkTInterface();
+			MikkTContext.m_pUserData = BuildData->GetMikkTUserData();
+			//MikkTContext.m_bIgnoreDegenerates = bIgnoreDegenerateTriangles;
+
+			genTangSpaceDefault(&MikkTContext);
+		}
+
+		check(WedgeTangentX.Num() == NumWedges);
+		check(WedgeTangentY.Num() == NumWedges);
+		check(WedgeTangentZ.Num() == NumWedges);
+	}
+
+	bool PrepareSourceMesh(IMeshBuildData* BuildData)
+	{
+		check(Stage == EStage::Uninit);
+
+		BeginSlowTask();
+
+		TMultiMap<int32, int32>& OverlappingCorners = *new(LODOverlappingCorners)TMultiMap<int32, int32>;
+
+		float ComparisonThreshold = THRESH_POINTS_ARE_SAME;//GetComparisonThreshold(LODBuildSettings[LODIndex]);
+		int32 NumWedges = BuildData->GetNumWedges();
+
+		// Find overlapping corners to accelerate adjacency.
+		Skeletal_FindOverlappingCorners(OverlappingCorners, BuildData, ComparisonThreshold);
+
+		// Figure out if we should recompute normals and tangents.
+		bool bRecomputeNormals = BuildData->BuildOptions.bComputeNormals;
+		bool bRecomputeTangents = BuildData->BuildOptions.bComputeTangents;
+
+		// Dump normals and tangents if we are recomputing them.
+		if (bRecomputeTangents)
+		{
+			TArray<FVector>& TangentX = BuildData->GetTangentArray(0);
+			TArray<FVector>& TangentY = BuildData->GetTangentArray(1);
+
+			TangentX.Empty(NumWedges);
+			TangentX.AddZeroed(NumWedges);
+			TangentY.Empty(NumWedges);
+			TangentY.AddZeroed(NumWedges);
+		}
+		if (bRecomputeNormals)
+		{
+			TArray<FVector>& TangentZ = BuildData->GetTangentArray(2);
+			TangentZ.Empty(NumWedges);
+			TangentZ.AddZeroed(NumWedges);
+		}
+
+		// Compute any missing tangents.
+		if (BuildData->BuildOptions.bUseMikkTSpace)
+		{
+			Skeletal_ComputeTangents_MikkTSpace(BuildData, OverlappingCorners);
+		}
+		else
+		{
+			Skeletal_ComputeTangents(BuildData, OverlappingCorners);
+		}
+
+		// At this point the mesh will have valid tangents.
+		BuildData->ValidateTangentArraySize();
+		check(LODOverlappingCorners.Num() == 1);
+
+		EndSlowTask();
+
+		Stage = EStage::Prepared;
+		return true;
+	}
+
+	bool GenerateSkeletalRenderMesh(IMeshBuildData* InBuildData)
+	{
+		check(Stage == EStage::Prepared);
+
+		SkeletalMeshBuildData& BuildData = *(SkeletalMeshBuildData*)InBuildData;
+
+		BeginSlowTask();
+	
+		// Find wedge influences.
+		TArray<int32>	WedgeInfluenceIndices;
+		TMap<uint32,uint32> VertexIndexToInfluenceIndexMap;
+
+		for(uint32 LookIdx = 0;LookIdx < (uint32)BuildData.Influences.Num();LookIdx++)
+		{
+			// Order matters do not allow the map to overwrite an existing value.
+			if( !VertexIndexToInfluenceIndexMap.Find( BuildData.Influences[LookIdx].VertIndex) )
+			{
+				VertexIndexToInfluenceIndexMap.Add( BuildData.Influences[LookIdx].VertIndex, LookIdx );
+			}
+		}
+
+		for(int32 WedgeIndex = 0;WedgeIndex < BuildData.Wedges.Num();WedgeIndex++)
+		{
+			uint32* InfluenceIndex = VertexIndexToInfluenceIndexMap.Find( BuildData.Wedges[WedgeIndex].iVertex );
+
+			if ( InfluenceIndex )
+			{
+				WedgeInfluenceIndices.Add( *InfluenceIndex );
+			}
+			else
+			{
+				// we have missing influence vert,  we weight to root
+				WedgeInfluenceIndices.Add(0);
+
+				// add warning message
+				if (BuildData.OutWarningMessages)
+				{
+					BuildData.OutWarningMessages->Add(FText::Format(FText::FromString("Missing influence on vert {0}. Weighting it to root."), FText::FromString(FString::FromInt(BuildData.Wedges[WedgeIndex].iVertex))));
+					if (BuildData.OutWarningNames)
+					{
+						BuildData.OutWarningNames->Add(FFbxErrors::SkeletalMesh_VertMissingInfluences);
+					}
+				}
+			}
+		}
+
+		check(BuildData.Wedges.Num() == WedgeInfluenceIndices.Num());
+
+		TArray<FSkeletalMeshVertIndexAndZ> VertIndexAndZ;
+		TArray<FSoftSkinBuildVertex> RawVertices;
+
+		VertIndexAndZ.Empty(BuildData.Points.Num());
+		RawVertices.Reserve(BuildData.Points.Num());
+
+		for(int32 FaceIndex = 0; FaceIndex < BuildData.Faces.Num(); FaceIndex++)
+		{
+			// Only update the status progress bar if we are in the game thread and every thousand faces. 
+			// Updating status is extremely slow
+			if (FaceIndex % 5000 == 0)
+			{
+				UpdateSlowTask(FaceIndex, BuildData.Faces.Num());
+			}
+
+			const FMeshFace& Face = BuildData.Faces[FaceIndex];
+						
+			for(int32 VertexIndex = 0; VertexIndex < 3; VertexIndex++)
+			{
+				FSoftSkinBuildVertex Vertex;
+				const uint32 WedgeIndex = BuildData.GetWedgeIndex(FaceIndex, VertexIndex);
+				const FMeshWedge& Wedge = BuildData.Wedges[WedgeIndex];
+
+				Vertex.Position = BuildData.GetVertexPosition(FaceIndex, VertexIndex);
+
+				FVector TangentX,TangentY,TangentZ;
+				TangentX = BuildData.TangentX[WedgeIndex].GetSafeNormal();
+				TangentY = BuildData.TangentY[WedgeIndex].GetSafeNormal();
+				TangentZ = BuildData.TangentZ[WedgeIndex].GetSafeNormal();
+
+				/*if (BuildData.BuildOptions.bComputeNormals || BuildData.BuildOptions.bComputeTangents)
+				{
+					TangentX = BuildData.TangentX[VertexIndex].GetSafeNormal();
+					TangentY = BuildData.TangentY[VertexIndex].GetSafeNormal();
+
+					if( BuildData.BuildOptions.bComputeNormals )
+					{
+						TangentZ = BuildData.TangentZ[VertexIndex].GetSafeNormal();
+					}
+					else
+					{
+						//TangentZ = Face.TangentZ[VertexIndex];
+					}
+
+					TangentY -= TangentX * (TangentX | TangentY);
+					TangentY.Normalize();
+
+					TangentX -= TangentZ * (TangentZ | TangentX);
+					TangentY -= TangentZ * (TangentZ | TangentY);
+
+					TangentX.Normalize();
+					TangentY.Normalize();
+				}
+				else*/
+				{
+					//TangentX = Face.TangentX[VertexIndex];
+					//TangentY = Face.TangentY[VertexIndex];
+					//TangentZ = Face.TangentZ[VertexIndex];
+
+					// Normalize overridden tangents.  Its possible for them to import un-normalized.
+					TangentX.Normalize();
+					TangentY.Normalize();
+					TangentZ.Normalize();
+				}
+
+				Vertex.TangentX = TangentX;
+				Vertex.TangentY = TangentY;
+				Vertex.TangentZ = TangentZ;
+
+				FMemory::Memcpy(Vertex.UVs, Wedge.UVs, sizeof(FVector2D)*MAX_TEXCOORDS);	
+				Vertex.Color = Wedge.Color;
+
+				{
+					// Count the influences.
+					int32 InfIdx = WedgeInfluenceIndices[Face.iWedge[VertexIndex]];
+					int32 LookIdx = InfIdx;
+
+					uint32 InfluenceCount = 0;
+					while( BuildData.Influences.IsValidIndex(LookIdx) && (BuildData.Influences[LookIdx].VertIndex == Wedge.iVertex) )
+					{			
+						InfluenceCount++;
+						LookIdx++;
+					}
+					InfluenceCount = FMath::Min<uint32>(InfluenceCount,MAX_TOTAL_INFLUENCES);
+
+					// Setup the vertex influences.
+					Vertex.InfluenceBones[0] = 0;
+					Vertex.InfluenceWeights[0] = 255;
+					for(uint32 i = 1;i < MAX_TOTAL_INFLUENCES;i++)
+					{
+						Vertex.InfluenceBones[i] = 0;
+						Vertex.InfluenceWeights[i] = 0;
+					}
+
+					uint32	TotalInfluenceWeight = 0;
+					for(uint32 i = 0;i < InfluenceCount;i++)
+					{
+						FBoneIndexType BoneIndex = (FBoneIndexType)BuildData.Influences[InfIdx+i].BoneIndex;
+						if( BoneIndex >= BuildData.RefSkeleton.GetNum() )
+							continue;
+
+						Vertex.InfluenceBones[i] = BoneIndex;
+						Vertex.InfluenceWeights[i] = (uint8)(BuildData.Influences[InfIdx+i].Weight * 255.0f);
+						TotalInfluenceWeight += Vertex.InfluenceWeights[i];
+					}
+					Vertex.InfluenceWeights[0] += 255 - TotalInfluenceWeight;
+				}
+
+				// Add the vertex as well as its original index in the points array
+				Vertex.PointWedgeIdx = Wedge.iVertex;
+			
+				int32 RawIndex = RawVertices.Add( Vertex );
+
+				// Add an efficient way to find dupes of this vertex later for fast combining of vertices
+				FSkeletalMeshVertIndexAndZ IAndZ;
+				IAndZ.Index = RawIndex;
+				IAndZ.Z = Vertex.Position.Z;
+
+				VertIndexAndZ.Add( IAndZ );
+			}
+		}
+
+		// Generate chunks and their vertices and indices
+		SkeletalMeshTools::BuildSkeletalMeshChunks(BuildData.Faces, RawVertices, VertIndexAndZ, BuildData.BuildOptions.bKeepOverlappingVertices, BuildData.Chunks, BuildData.bTooManyVerts);
+
+		// Chunk vertices to satisfy the requested limit.
+		static const auto MaxBonesVar = IConsoleManager::Get().FindTConsoleVariableDataInt(TEXT("Compat.MAX_GPUSKIN_BONES"));
+		const int32 MaxGPUSkinBones = MaxBonesVar->GetValueOnAnyThread();
+		SkeletalMeshTools::ChunkSkinnedVertices(BuildData.Chunks, MaxGPUSkinBones);
+
+		EndSlowTask();
+
+		Stage = EStage::GenerateRendering;
+		return true;
+	}
+
+	void BeginSlowTask()
+	{
+		if (IsInGameThread())
+		{
+			GWarn->BeginSlowTask( NSLOCTEXT("UnrealEd", "ProcessingSkeletalTriangles", "Processing Mesh Triangles"), true );
+		}
+	}
+
+	void UpdateSlowTask(int32 Numerator, int32 Denominator)
+	{
+		if (IsInGameThread())
+		{
+			GWarn->StatusUpdate(Numerator, Denominator, NSLOCTEXT("UnrealEd","ProcessingSkeletalTriangles","Processing Mesh Triangles"));
+		}
+	}
+
+	void EndSlowTask()
+	{
+		if (IsInGameThread())
+		{
+			GWarn->EndSlowTask();
+		}
+	}
+
+private:
+	enum class EStage
+	{
+		Uninit,
+		Prepared,
+		GenerateRendering,
+	};
+
+	TIndirectArray<TMultiMap<int32, int32> > LODOverlappingCorners;
+	EStage Stage;
+};
+
+bool FMeshUtilities::BuildSkeletalMesh(FStaticLODModel& LODModel, const FReferenceSkeleton& RefSkeleton, const TArray<FVertInfluence>& Influences, const TArray<FMeshWedge>& Wedges, const TArray<FMeshFace>& Faces, const TArray<FVector>& Points, const TArray<int32>& PointToOriginalMap, const MeshBuildOptions& BuildOptions, TArray<FText> * OutWarningMessages, TArray<FName> * OutWarningNames)
 {
 #if WITH_EDITORONLY_DATA
+	// Temporarily supporting both import paths
+	if (!BuildOptions.bUseMikkTSpace)
+	{
+		return BuildSkeletalMesh_Legacy(LODModel, RefSkeleton, Influences, Wedges, Faces, Points, PointToOriginalMap, BuildOptions.bKeepOverlappingVertices, BuildOptions.bComputeNormals, BuildOptions.bComputeTangents, OutWarningMessages, OutWarningNames);
+	}
+
+	SkeletalMeshBuildData BuildData(
+		LODModel,
+		RefSkeleton,
+		Influences,
+		Wedges,
+		Faces,
+		Points,
+		PointToOriginalMap,
+		BuildOptions,
+		OutWarningMessages,
+		OutWarningNames);
+
+	FSkeletalMeshUtilityBuilder Builder;
+	if (!Builder.PrepareSourceMesh(&BuildData))
+	{
+		return false;
+	}
+
+	if (!Builder.GenerateSkeletalRenderMesh(&BuildData))
+	{
+		return false;
+	}
+
+	// Build the skeletal model from chunks.
+	Builder.BeginSlowTask();
+	BuildSkeletalModelFromChunks(BuildData.LODModel, BuildData.RefSkeleton, BuildData.Chunks, BuildData.PointToOriginalMap);
+	Builder.EndSlowTask();
+
+	// Only show these warnings if in the game thread.  When importing morph targets, this function can run in another thread and these warnings dont prevent the mesh from importing
+	if( IsInGameThread() )
+	{
+		bool bHasBadSections = false;
+		for (int32 SectionIndex = 0; SectionIndex < BuildData.LODModel.Sections.Num(); SectionIndex++)
+		{
+			FSkelMeshSection& Section = BuildData.LODModel.Sections[SectionIndex];
+			bHasBadSections |= (Section.NumTriangles == 0);
+
+			// Log info about the section.
+			UE_LOG(LogSkeletalMesh, Log, TEXT("Section %u: Material=%u, Chunk=%u, %u triangles"),
+				SectionIndex,
+				Section.MaterialIndex,
+				Section.ChunkIndex,
+				Section.NumTriangles
+				);
+		}
+		if( bHasBadSections )
+		{
+			FText BadSectionMessage(NSLOCTEXT("UnrealEd", "Error_SkeletalMeshHasBadSections", "Input mesh has a section with no triangles.  This mesh may not render properly."));
+			if(BuildData.OutWarningMessages)
+			{
+				BuildData.OutWarningMessages->Add(BadSectionMessage);
+				if(BuildData.OutWarningNames)
+				{
+					BuildData.OutWarningNames->Add(FFbxErrors::SkeletalMesh_SectionWithNoTriangle);
+				}
+			}
+			else
+			{
+				FMessageDialog::Open( EAppMsgType::Ok, BadSectionMessage );
+			}
+		}
+
+		if (BuildData.bTooManyVerts)
+		{
+			FText TooManyVertsMessage(NSLOCTEXT("UnrealEd", "Error_SkeletalMeshTooManyVertices", "Input mesh has too many vertices.  The generated mesh will be corrupt!  Consider adding extra materials to split up the source mesh into smaller chunks."));
+
+			if(BuildData.OutWarningMessages)
+			{
+				BuildData.OutWarningMessages->Add(TooManyVertsMessage);
+				if(BuildData.OutWarningNames)
+				{
+					BuildData.OutWarningNames->Add(FFbxErrors::SkeletalMesh_TooManyVertices);
+				}
+			}
+			else
+			{
+				FMessageDialog::Open(EAppMsgType::Ok, TooManyVertsMessage);
+			}
+		}
+	}
+
+	return true;
+#else
+	if(OutWarningMessages)
+	{
+		OutWarningMessages->Add(FText::FromString(TEXT("Cannot call FMeshUtilities::BuildSkeletalMesh on a console!")));
+	}
+	else
+	{
+		UE_LOG(LogSkeletalMesh, Fatal,TEXT("Cannot call FMeshUtilities::BuildSkeletalMesh on a console!"));
+	}
+	return false;
+#endif
+}
+
+//@TODO: The OutMessages has to be a struct that contains FText/FName, or make it Token and add that as error. Needs re-work. Temporary workaround for now. 
+bool FMeshUtilities::BuildSkeletalMesh_Legacy( FStaticLODModel& LODModel, const FReferenceSkeleton& RefSkeleton, const TArray<FVertInfluence>& Influences, const TArray<FMeshWedge>& Wedges, const TArray<FMeshFace>& Faces, const TArray<FVector>& Points, const TArray<int32>& PointToOriginalMap, bool bKeepOverlappingVertices, bool bComputeNormals, bool bComputeTangents, TArray<FText> * OutWarningMessages, TArray<FName> * OutWarningNames)
+{
 	bool bTooManyVerts = false;
 
 	check(PointToOriginalMap.Num() == Points.Num());
@@ -3636,18 +4897,6 @@ bool FMeshUtilities::BuildSkeletalMesh( FStaticLODModel& LODModel, const FRefere
 	}
 
 	return true;
-#else
-
-	if(OutWarningMessages)
-	{
-		OutWarningMessages->Add(FText::FromString(TEXT("Cannot call FSkeletalMeshTools::CreateSkinningStreams on a console!")));
-	}
-	else
-	{
-		UE_LOG(LogSkeletalMesh, Fatal,TEXT("Cannot call FSkeletalMeshTools::CreateSkinningStreams on a console!"));
-	}
-	return false;
-#endif
 }
 
 static bool NonOpaqueMaterialPredicate(UStaticMeshComponent* InMesh)
diff --git a/Engine/Source/Developer/MeshUtilities/Public/MeshUtilities.h b/Engine/Source/Developer/MeshUtilities/Public/MeshUtilities.h
--- a/Engine/Source/Developer/MeshUtilities/Public/MeshUtilities.h
+++ b/Engine/Source/Developer/MeshUtilities/Public/MeshUtilities.h
@@ -115,6 +115,25 @@ public:
 		bool bGenerateAsIfTwoSided,
 		class FDistanceFieldVolumeData& OutData) = 0;
 
+	/** Helper structure for skeletal mesh import options */
+	struct MeshBuildOptions
+	{
+		MeshBuildOptions()
+		: bKeepOverlappingVertices(false)
+		, bRemoveDegenerateTriangles(false)
+		, bComputeNormals(true)
+		, bComputeTangents(true)
+		, bUseMikkTSpace(false)
+		{
+		}
+
+		bool bKeepOverlappingVertices;
+		bool bRemoveDegenerateTriangles;
+		bool bComputeNormals;
+		bool bComputeTangents;
+		bool bUseMikkTSpace;
+	};
+	
 	/**
 	 * Create all render specific data for a skeletal mesh LOD model
 	 * @returns true if the mesh was built successfully.
@@ -127,9 +146,7 @@ public:
 		const TArray<FMeshFace>& Faces, 
 		const TArray<FVector>& Points,
 		const TArray<int32>& PointToOriginalMap,
-		bool bKeepOverlappingVertices = false,
-		bool bComputeNormals = true,
-		bool bComputeTangents = true,
+		const MeshBuildOptions& BuildOptions = MeshBuildOptions(),
 		TArray<FText> * OutWarningMessages = NULL,
 		TArray<FName> * OutWarningNames = NULL
 		) = 0;
diff --git a/Engine/Source/Editor/UnrealEd/Classes/Factories/FbxMeshImportData.h b/Engine/Source/Editor/UnrealEd/Classes/Factories/FbxMeshImportData.h
--- a/Engine/Source/Editor/UnrealEd/Classes/Factories/FbxMeshImportData.h
+++ b/Engine/Source/Editor/UnrealEd/Classes/Factories/FbxMeshImportData.h
@@ -45,5 +45,9 @@ class UFbxMeshImportData : public UFbxAssetImportData
 	UPROPERTY(EditAnywhere, config, AdvancedDisplay, Category = ImportSettings, meta = (ImportType="StaticMesh"))
 	TEnumAsByte<enum EFBXNormalGenerationMethod::Type> NormalGenerationMethod;
 
+	/** Enables experimental Mikk tangent generation for skeletal meshes */
+	UPROPERTY(EditAnywhere, AdvancedDisplay, config, Category=ImportSettings, meta=(ImportType="SkeletalMesh"))
+	uint32 bUseExperimentalTangentGeneration:1;
+
 	bool CanEditChange( const UProperty* InProperty ) const override;
 };
diff --git a/Engine/Source/Editor/UnrealEd/Private/Factories/ApexDestructibleAssetImport.cpp b/Engine/Source/Editor/UnrealEd/Private/Factories/ApexDestructibleAssetImport.cpp
--- a/Engine/Source/Editor/UnrealEd/Private/Factories/ApexDestructibleAssetImport.cpp
+++ b/Engine/Source/Editor/UnrealEd/Private/Factories/ApexDestructibleAssetImport.cpp
@@ -858,8 +858,13 @@ bool SetApexDestructibleAsset(UDestructibleMesh& DestructibleMesh, NxDestructibl
 
 		IMeshUtilities& MeshUtilities = FModuleManager::Get().LoadModuleChecked<IMeshUtilities>("MeshUtilities");
 
+		IMeshUtilities::MeshBuildOptions BuildOptions;
+		BuildOptions.bKeepOverlappingVertices = false;
+		BuildOptions.bComputeNormals = !bHaveNormals;
+		BuildOptions.bComputeTangents = !bHaveTangents;
+
 		// Create actual rendering data.
-		if (!MeshUtilities.BuildSkeletalMesh(DestructibleMeshResource.LODModels[0], DestructibleMesh.RefSkeleton, LODInfluences,LODWedges,LODFaces,LODPoints,LODPointToRawMap,false,!bHaveNormals,!bHaveTangents))
+		if (!MeshUtilities.BuildSkeletalMesh(DestructibleMeshResource.LODModels[0], DestructibleMesh.RefSkeleton, LODInfluences,LODWedges,LODFaces,LODPoints,LODPointToRawMap,BuildOptions))
 		{
 			DestructibleMesh.MarkPendingKill();
 			return false;
diff --git a/Engine/Source/Editor/UnrealEd/Private/Fbx/FbxMainImport.cpp b/Engine/Source/Editor/UnrealEd/Private/Fbx/FbxMainImport.cpp
--- a/Engine/Source/Editor/UnrealEd/Private/Fbx/FbxMainImport.cpp
+++ b/Engine/Source/Editor/UnrealEd/Private/Fbx/FbxMainImport.cpp
@@ -239,6 +239,7 @@ void ApplyImportUIToImportOptions(UFbxImportUI* ImportUI, FBXImportOptions& InOu
 	InOutImportOptions.bImportMeshesInBoneHierarchy = ImportUI->SkeletalMeshImportData->bImportMeshesInBoneHierarchy;
 	InOutImportOptions.bCreatePhysicsAsset = ImportUI->bCreatePhysicsAsset;
 	InOutImportOptions.PhysicsAsset = ImportUI->PhysicsAsset;
+	InOutImportOptions.bUseExperimentalTangentGeneration = ImportUI->SkeletalMeshImportData->bUseExperimentalTangentGeneration;
 	// animation options
 	InOutImportOptions.AnimationLengthImportType = ImportUI->AnimSequenceImportData->AnimationLength;
 	InOutImportOptions.AnimationRange.X = ImportUI->AnimSequenceImportData->StartFrame;
diff --git a/Engine/Source/Editor/UnrealEd/Private/Fbx/FbxSkeletalMeshImport.cpp b/Engine/Source/Editor/UnrealEd/Private/Fbx/FbxSkeletalMeshImport.cpp
--- a/Engine/Source/Editor/UnrealEd/Private/Fbx/FbxSkeletalMeshImport.cpp
+++ b/Engine/Source/Editor/UnrealEd/Private/Fbx/FbxSkeletalMeshImport.cpp
@@ -390,7 +390,7 @@ bool UnFbx::FFbxImporter::FacesAreSmoothlyConnected( FSkeletalMeshImportData &Im
 	return ( VertMatches >= 2 );
 }
 
-int32	UnFbx::FFbxImporter::DoUnSmoothVerts(FSkeletalMeshImportData &ImportData)
+int32	UnFbx::FFbxImporter::DoUnSmoothVerts(FSkeletalMeshImportData &ImportData, bool bDuplicateUnSmoothWedges)
 {
 	//
 	// Connectivity: triangles with non-matching smoothing groups will be physically split.
@@ -565,21 +565,29 @@ int32	UnFbx::FFbxImporter::DoUnSmoothVerts(FSkeletalMeshImportData &ImportData)
 					{
 						int32 OldWedgeIndex = PointWedges[p].WedgeList[w];
 						int32 NewWedgeIndex = ImportData.Wedges.Num();
-						ImportData.Wedges.AddUninitialized();
-						ImportData.Wedges[NewWedgeIndex] = ImportData.Wedges[ OldWedgeIndex ];
-						ImportData.Wedges[ NewWedgeIndex ].VertexIndex = NewPointIndex; 
 
-						//  Update relevant face's Wedges. Inelegant: just check all associated faces for every new wedge.
-						for( int32 s=0; s< FaceSets[f].Faces.Num(); s++)
+						if( bDuplicateUnSmoothWedges )
 						{
-							int32 FanIndex = FaceSets[f].Faces[s];
-							if( Fans[p].FaceRecord[ FanIndex ].WedgeIndex == OldWedgeIndex )
+							ImportData.Wedges.AddUninitialized();
+							ImportData.Wedges[NewWedgeIndex] = ImportData.Wedges[ OldWedgeIndex ];
+							ImportData.Wedges[ NewWedgeIndex ].VertexIndex = NewPointIndex; 
+
+							//  Update relevant face's Wedges. Inelegant: just check all associated faces for every new wedge.
+							for( int32 s=0; s< FaceSets[f].Faces.Num(); s++)
 							{
-								// Update just the right one for this face (HoekIndex!) 
-								ImportData.Faces[ Fans[p].FaceRecord[ FanIndex].FaceIndex ].WedgeIndex[ Fans[p].FaceRecord[ FanIndex ].HoekIndex ] = NewWedgeIndex;
-								RemappedHoeks++;
+								int32 FanIndex = FaceSets[f].Faces[s];
+								if( Fans[p].FaceRecord[ FanIndex ].WedgeIndex == OldWedgeIndex )
+								{
+									// Update just the right one for this face (HoekIndex!) 
+									ImportData.Faces[ Fans[p].FaceRecord[ FanIndex].FaceIndex ].WedgeIndex[ Fans[p].FaceRecord[ FanIndex ].HoekIndex ] = NewWedgeIndex;
+									RemappedHoeks++;
+								}
 							}
 						}
+						else
+						{
+							ImportData.Wedges[OldWedgeIndex].VertexIndex = NewPointIndex; 
+						}
 					}
 				}
 			} //  if FaceSets.Num(). -> duplicate stuff
@@ -1206,7 +1214,8 @@ USkeletalMesh* UnFbx::FFbxImporter::ImportSkeletalMesh(UObject* InParent, TArray
 	
 	if( ImportOptions->bPreserveSmoothingGroups )
 	{
-		DoUnSmoothVerts(*SkelMeshImportDataPtr);
+		bool bDuplicateUnSmoothWedges = !ImportOptions->bUseExperimentalTangentGeneration;
+		DoUnSmoothVerts(*SkelMeshImportDataPtr, bDuplicateUnSmoothWedges);
 	}
 	else
 	{
@@ -1275,40 +1284,36 @@ USkeletalMesh* UnFbx::FFbxImporter::ImportSkeletalMesh(UObject* InParent, TArray
 		TArray<int32> LODPointToRawMap;
 		SkelMeshImportDataPtr->CopyLODImportData(LODPoints,LODWedges,LODFaces,LODInfluences,LODPointToRawMap);
 
-		const bool bShouldComputeNormals = !ImportOptions->ShouldImportNormals() || !SkelMeshImportDataPtr->bHasNormals;
-		const bool bShouldComputeTangents = !ImportOptions->ShouldImportTangents() || !SkelMeshImportDataPtr->bHasTangents;
+		IMeshUtilities::MeshBuildOptions BuildOptions;
+		BuildOptions.bKeepOverlappingVertices = ImportOptions->bKeepOverlappingVertices;
+		BuildOptions.bComputeNormals = !ImportOptions->ShouldImportNormals() || !SkelMeshImportDataPtr->bHasNormals;
+		BuildOptions.bComputeTangents = !ImportOptions->ShouldImportTangents() || !SkelMeshImportDataPtr->bHasTangents;
+		BuildOptions.bUseMikkTSpace = ImportOptions->bUseExperimentalTangentGeneration;
+		BuildOptions.bRemoveDegenerateTriangles = false;
 
 		IMeshUtilities& MeshUtilities = FModuleManager::Get().LoadModuleChecked<IMeshUtilities>("MeshUtilities");
 		
 		TArray<FText> WarningMessages;
 		TArray<FName> WarningNames;
 		// Create actual rendering data.
-		if ( !MeshUtilities.BuildSkeletalMesh(ImportedResource->LODModels[0],SkeletalMesh->RefSkeleton,LODInfluences,LODWedges,LODFaces,LODPoints,LODPointToRawMap, ImportOptions->bKeepOverlappingVertices, bShouldComputeNormals, bShouldComputeTangents, &WarningMessages, &WarningNames ) )
+		bool bBuildSuccess = MeshUtilities.BuildSkeletalMesh(ImportedResource->LODModels[0],SkeletalMesh->RefSkeleton,LODInfluences,LODWedges,LODFaces,LODPoints,LODPointToRawMap,BuildOptions, &WarningMessages, &WarningNames);
+
+		// temporary hack of message/names, should be one token or a struct
+		if(WarningMessages.Num() > 0 && WarningNames.Num() == WarningMessages.Num())
 		{
-			if (WarningNames.Num() == WarningMessages.Num())
+			EMessageSeverity::Type MessageSeverity = bBuildSuccess ? EMessageSeverity::Warning : EMessageSeverity::Error;
+
+			for(int32 MessageIdx = 0; MessageIdx<WarningMessages.Num(); ++MessageIdx)
 			{
-				// temporary hack of message/names, should be one token or a struct
-				for(int32 MessageIdx = 0; MessageIdx<WarningMessages.Num(); ++MessageIdx)
-				{
-					AddTokenizedErrorMessage(FTokenizedMessage::Create(EMessageSeverity::Error, WarningMessages[MessageIdx]), WarningNames[MessageIdx]);
-				}
+				AddTokenizedErrorMessage(FTokenizedMessage::Create(MessageSeverity, WarningMessages[MessageIdx]), WarningNames[MessageIdx]);
 			}
+		}
 
+		if( !bBuildSuccess )
+		{
 			SkeletalMesh->MarkPendingKill();
 			return NULL;
 		}
-		else if (WarningMessages.Num() > 0)
-		{
-			// temporary hack of message/names, should be one token or a struct
-			if(WarningNames.Num() == WarningMessages.Num())
-			{
-				// temporary hack of message/names, should be one token or a struct
-				for(int32 MessageIdx = 0; MessageIdx<WarningMessages.Num(); ++MessageIdx)
-				{
-					AddTokenizedErrorMessage(FTokenizedMessage::Create(EMessageSeverity::Warning, WarningMessages[MessageIdx]), WarningNames[MessageIdx]);
-				}
-			}
-		}
 
 		// Presize the per-section shadow casting array with the number of sections in the imported LOD.
 		const int32 NumSections = LODModel.Sections.Num();
@@ -2684,9 +2689,12 @@ public:
 		TArray<FVertInfluence> LODInfluences;
 		TArray<int32> LODPointToRawMap;
 		ImportData.CopyLODImportData(LODPoints,LODWedges,LODFaces,LODInfluences,LODPointToRawMap);
+
+		IMeshUtilities::MeshBuildOptions BuildOptions;
+		BuildOptions.bKeepOverlappingVertices = bKeepOverlappingVertices;
 	
 		ImportData.Empty();
-		MeshUtilities->BuildSkeletalMesh( TempSkeletalMesh->GetImportedResource()->LODModels[0], TempSkeletalMesh->RefSkeleton, LODInfluences, LODWedges, LODFaces, LODPoints, LODPointToRawMap, bKeepOverlappingVertices);
+		MeshUtilities->BuildSkeletalMesh( TempSkeletalMesh->GetImportedResource()->LODModels[0], TempSkeletalMesh->RefSkeleton, LODInfluences, LODWedges, LODFaces, LODPoints, LODPointToRawMap, BuildOptions);
 	}
 
 	FORCEINLINE TStatId GetStatId() const
diff --git a/Engine/Source/Editor/UnrealEd/Private/FbxImporter.h b/Engine/Source/Editor/UnrealEd/Private/FbxImporter.h
--- a/Engine/Source/Editor/UnrealEd/Private/FbxImporter.h
+++ b/Engine/Source/Editor/UnrealEd/Private/FbxImporter.h
@@ -141,6 +141,7 @@ struct FBXImportOptions
 	bool bKeepOverlappingVertices;
 	bool bImportMeshesInBoneHierarchy;
 	bool bCreatePhysicsAsset;
+	bool bUseExperimentalTangentGeneration;
 	UPhysicsAsset *PhysicsAsset;
 	// Animation option
 	USkeleton* SkeletonForAnimation;
@@ -925,7 +926,7 @@ protected:
 	 * @param ImportData
 	 * @return int32 number of points that added when process unsmooth faces
 	*/
-	int32 DoUnSmoothVerts(FSkeletalMeshImportData &ImportData);
+	int32 DoUnSmoothVerts(FSkeletalMeshImportData &ImportData, bool bDuplicateUnSmoothWedges = true);
 	
 	/**
 	 * Merge all layers of one AnimStack to one layer.
