Copying //UE4/Dev-Framework to Dev-Main (//UE4/Dev-Main)

diff --git a/Engine/Source/Developer/MeshUtilities/Private/MeshUtilities.cpp b/Engine/Source/Developer/MeshUtilities/Private/MeshUtilities.cpp
--- a/Engine/Source/Developer/MeshUtilities/Private/MeshUtilities.cpp
+++ b/Engine/Source/Developer/MeshUtilities/Private/MeshUtilities.cpp
@@ -33,7 +33,7 @@
 #endif
 
 /*------------------------------------------------------------------------------
-	MeshUtilities module.
+MeshUtilities module.
 ------------------------------------------------------------------------------*/
 
 // The version string is a GUID. If you make a change to mesh utilities that
@@ -42,7 +42,7 @@
 
 #define MESH_UTILITIES_VER TEXT("8C68575CEF434CA8A9E1DA4AED8A47BB")
 
-DEFINE_LOG_CATEGORY_STATIC(LogMeshUtilities,Verbose,All);
+DEFINE_LOG_CATEGORY_STATIC(LogMeshUtilities, Verbose, All);
 
 #define LOCTEXT_NAMESPACE "MeshUtils"
 
@@ -102,22 +102,22 @@ private:
 
 	virtual void GenerateSignedDistanceFieldVolumeData(
 		const FStaticMeshLODResources& LODModel,
-		class FQueuedThreadPool& ThreadPool,
+	class FQueuedThreadPool& ThreadPool,
 		const TArray<EBlendMode>& MaterialBlendModes,
 		const FBoxSphereBounds& Bounds,
 		float DistanceFieldResolutionScale,
 		bool bGenerateAsIfTwoSided,
 		FDistanceFieldVolumeData& OutData) override;
 
-	virtual bool BuildSkeletalMesh( FStaticLODModel& LODModel, const FReferenceSkeleton& RefSkeleton, const TArray<FVertInfluence>& Influences, const TArray<FMeshWedge>& Wedges, const TArray<FMeshFace>& Faces, const TArray<FVector>& Points, const TArray<int32>& PointToOriginalMap, const MeshBuildOptions& BuildOptions = MeshBuildOptions(), TArray<FText> * OutWarningMessages = NULL, TArray<FName> * OutWarningNames = NULL) override;
-	bool BuildSkeletalMesh_Legacy( FStaticLODModel& LODModel, const FReferenceSkeleton& RefSkeleton, const TArray<FVertInfluence>& Influences, const TArray<FMeshWedge>& Wedges, const TArray<FMeshFace>& Faces, const TArray<FVector>& Points, const TArray<int32>& PointToOriginalMap, bool bKeepOverlappingVertices = false, bool bComputeNormals = true, bool bComputeTangents = true, TArray<FText> * OutWarningMessages = NULL, TArray<FName> * OutWarningNames = NULL);
+	virtual bool BuildSkeletalMesh(FStaticLODModel& LODModel, const FReferenceSkeleton& RefSkeleton, const TArray<FVertInfluence>& Influences, const TArray<FMeshWedge>& Wedges, const TArray<FMeshFace>& Faces, const TArray<FVector>& Points, const TArray<int32>& PointToOriginalMap, const MeshBuildOptions& BuildOptions = MeshBuildOptions(), TArray<FText> * OutWarningMessages = NULL, TArray<FName> * OutWarningNames = NULL) override;
+	bool BuildSkeletalMesh_Legacy(FStaticLODModel& LODModel, const FReferenceSkeleton& RefSkeleton, const TArray<FVertInfluence>& Influences, const TArray<FMeshWedge>& Wedges, const TArray<FMeshFace>& Faces, const TArray<FVector>& Points, const TArray<int32>& PointToOriginalMap, bool bKeepOverlappingVertices = false, bool bComputeNormals = true, bool bComputeTangents = true, TArray<FText> * OutWarningMessages = NULL, TArray<FName> * OutWarningNames = NULL);
 
 	virtual IMeshReduction* GetMeshReductionInterface() override;
 	virtual IMeshMerging* GetMeshMergingInterface() override;
 	virtual void CacheOptimizeIndexBuffer(TArray<uint16>& Indices) override;
 	virtual void CacheOptimizeIndexBuffer(TArray<uint32>& Indices) override;
-	void CacheOptimizeVertexAndIndexBuffer(TArray<FStaticMeshBuildVertex>& Vertices,TArray<TArray<uint32> >& PerSectionIndices,TArray<int32>& WedgeMap);
-	
+	void CacheOptimizeVertexAndIndexBuffer(TArray<FStaticMeshBuildVertex>& Vertices, TArray<TArray<uint32> >& PerSectionIndices, TArray<int32>& WedgeMap);
+
 	virtual void BuildSkeletalAdjacencyIndexBuffer(
 		const TArray<FSoftSkinVertex>& VertexBuffer,
 		const uint32 TexCoordCount,
@@ -130,13 +130,13 @@ private:
 	virtual void CalcBoneVertInfos(USkeletalMesh* SkeletalMesh, TArray<FBoneVertInfo>& Infos, bool bOnlyDominant) override;
 
 	/**
-	 * Builds a renderable skeletal mesh LOD model. Note that the array of chunks
-	 * will be destroyed during this process!
-	 * @param LODModel				Upon return contains a renderable skeletal mesh LOD model.
-	 * @param RefSkeleton			The reference skeleton associated with the model.
-	 * @param Chunks				Skinned mesh chunks from which to build the renderable model.
-	 * @param PointToOriginalMap	Maps a vertex's RawPointIdx to its index at import time.
-	 */
+	* Builds a renderable skeletal mesh LOD model. Note that the array of chunks
+	* will be destroyed during this process!
+	* @param LODModel				Upon return contains a renderable skeletal mesh LOD model.
+	* @param RefSkeleton			The reference skeleton associated with the model.
+	* @param Chunks				Skinned mesh chunks from which to build the renderable model.
+	* @param PointToOriginalMap	Maps a vertex's RawPointIdx to its index at import time.
+	*/
 	void BuildSkeletalModelFromChunks(FStaticLODModel& LODModel, const FReferenceSkeleton& RefSkeleton, TArray<FSkinnedMeshChunk*>& Chunks, const TArray<int32>& PointToOriginalMap);
 
 	// IModuleInterface interface.
@@ -165,11 +165,11 @@ private:
 		FVector& OutMergedActorLocation,
 		const float ScreenAreaSize,
 		bool bSilent = false) const override;
-	
+
 	virtual void CreateProxyMesh(const TArray<AActor*>& InActors, const struct FMeshProxySettings& InMeshProxySettings, UPackage* InOuter, const FString& InProxyBasePackageName, const FGuid InGuid, FCreateProxyDelegate InProxyCreatedDelegate, const bool bAllowAsync) override;
 
 	DEPRECATED(4.11, "Please use CreateProxyMesh with new signature")
-	virtual void CreateProxyMesh(
+		virtual void CreateProxyMesh(
 		const TArray<AActor*>& Actors,
 		const struct FMeshProxySettings& InProxySettings,
 		UPackage* InOuter,
@@ -187,17 +187,17 @@ private:
 
 
 	virtual void FlattenMaterialsWithMeshData(TArray<UMaterialInterface*>& InMaterials, TArray<FRawMeshExt>& InSourceMeshes, TMap<FMeshIdAndLOD, TArray<int32>>& InMaterialIndexMap, TArray<bool>& InMeshShouldBakeVertexData, const FMaterialProxySettings &InMaterialProxySettings, TArray<FFlattenMaterial> &OutFlattenedMaterials) const override;
-	
+
 
 	bool ConstructRawMesh(
 		UStaticMeshComponent* MeshComponent,
 		int32 LODIndex,
-		FRawMesh& OutRawMesh, 
+		FRawMesh& OutRawMesh,
 		TArray<UMaterialInterface*>& OutUniqueMaterials,
 		TArray<int32>& OutGlobalMaterialIndices
 		) const;
 
-	virtual void ExtractMeshDataForGeometryCache(FRawMesh& RawMesh, const FMeshBuildSettings& BuildSettings, TArray<FStaticMeshBuildVertex>& OutVertices, TArray<TArray<uint32> >& OutPerSectionIndices );
+	virtual void ExtractMeshDataForGeometryCache(FRawMesh& RawMesh, const FMeshBuildSettings& BuildSettings, TArray<FStaticMeshBuildVertex>& OutVertices, TArray<TArray<uint32> >& OutPerSectionIndices);
 
 	virtual bool PropagatePaintedColorsToRawMesh(UStaticMeshComponent* StaticMeshComponent, int32 LODIndex, FRawMesh& RawMesh) const override;
 
@@ -246,7 +246,7 @@ public:
 	{
 		FScopeLock Lock(&StateLock);
 		for (const auto& Entry : ToProcessJobDataMap)
-		{			
+		{
 			FGuid JobGuid = Entry.Key;
 			FProxyGenerationData* Data = Entry.Value;
 
@@ -270,7 +270,7 @@ public:
 		FMergeCompleteData** FindData = ProxyMeshJobs.Find(OutJobGUID);
 		if (FindData && *FindData)
 		{
-			FMergeCompleteData* Data = *FindData;			
+			FMergeCompleteData* Data = *FindData;
 
 			FProxyGenerationData* GenerationData = new FProxyGenerationData();
 			GenerationData->Material = OutMaterial;
@@ -336,6 +336,9 @@ protected:
 		TArray<UObject*> OutAssetsToSync;
 		const FString AssetBaseName = FPackageName::GetShortName(Data->MergeData->ProxyBasePackageName);
 		const FString AssetBasePath = Data->MergeData->InOuter ? TEXT("") : FPackageName::GetLongPackagePath(Data->MergeData->ProxyBasePackageName) + TEXT("/");
+		
+		// Resize flattened material
+		FMaterialUtilities::ResizeFlattenMaterial(Data->Material, Data->MergeData->InProxySettings);
 
 		// Optimize flattened material
 		FMaterialUtilities::OptimizeFlattenMaterial(Data->Material);
@@ -412,13 +415,13 @@ public:
 
 	/** Initialization constructor. */
 	FMeshBuildDataProvider(
-		const TkDOPTree<const FMeshBuildDataProvider,uint32>& InkDopTree):
+		const TkDOPTree<const FMeshBuildDataProvider, uint32>& InkDopTree) :
 		kDopTree(InkDopTree)
 	{}
 
 	// kDOP data provider interface.
 
-	FORCEINLINE const TkDOPTree<const FMeshBuildDataProvider,uint32>& GetkDOPTree(void) const
+	FORCEINLINE const TkDOPTree<const FMeshBuildDataProvider, uint32>& GetkDOPTree(void) const
 	{
 		return kDopTree;
 	}
@@ -445,7 +448,7 @@ public:
 
 private:
 
-	const TkDOPTree<const FMeshBuildDataProvider,uint32>& kDopTree;
+	const TkDOPTree<const FMeshBuildDataProvider, uint32>& kDopTree;
 };
 
 /** Generates unit length, stratified and uniformly distributed direction samples in a hemisphere. */
@@ -474,7 +477,7 @@ void GenerateStratifiedUniformHemisphereSamples(int32 NumThetaSteps, int32 NumPh
 class FMeshDistanceFieldAsyncTask : public FNonAbandonableTask
 {
 public:
-	FMeshDistanceFieldAsyncTask(TkDOPTree<const FMeshBuildDataProvider,uint32>* InkDopTree, 
+	FMeshDistanceFieldAsyncTask(TkDOPTree<const FMeshBuildDataProvider, uint32>* InkDopTree,
 		const TArray<FVector4>* InSampleDirections,
 		FBox InVolumeBounds,
 		FIntVector InVolumeDimensions,
@@ -499,7 +502,7 @@ public:
 		RETURN_QUICK_DECLARE_CYCLE_STAT(FMeshDistanceFieldAsyncTask, STATGROUP_ThreadPoolAsyncTasks);
 	}
 
-	bool WasNegativeAtBorder() const 
+	bool WasNegativeAtBorder() const
 	{
 		return bNegativeAtBorder;
 	}
@@ -507,7 +510,7 @@ public:
 private:
 
 	// Readonly inputs
-	TkDOPTree<const FMeshBuildDataProvider,uint32>* kDopTree;
+	TkDOPTree<const FMeshBuildDataProvider, uint32>* kDopTree;
 	const TArray<FVector4>* SampleDirections;
 	FBox VolumeBounds;
 	FIntVector VolumeDimensions;
@@ -544,7 +547,7 @@ void FMeshDistanceFieldAsyncTask::DoWork()
 				{
 					FkHitResult Result;
 
-					TkDOPLineCollisionCheck<const FMeshBuildDataProvider,uint32> kDOPCheck(
+					TkDOPLineCollisionCheck<const FMeshBuildDataProvider, uint32> kDOPCheck(
 						VoxelPosition,
 						VoxelPosition + RayDirection * VolumeMaxDistance,
 						true,
@@ -559,7 +562,7 @@ void FMeshDistanceFieldAsyncTask::DoWork()
 
 						const FVector HitNormal = kDOPCheck.GetHitNormal();
 
-						if (FVector::DotProduct(RayDirection, HitNormal) > 0 
+						if (FVector::DotProduct(RayDirection, HitNormal) > 0
 							// MaterialIndex on the build triangles was set to 1 if two-sided, or 0 if one-sided
 							&& kDOPCheck.Result->Item == 0)
 						{
@@ -590,9 +593,9 @@ void FMeshDistanceFieldAsyncTask::DoWork()
 
 			const float VolumeSpaceDistance = MinDistance / VolumeBounds.GetExtent().GetMax();
 
-			if (MinDistance < 0 && 
-				(XIndex == 0 || XIndex == VolumeDimensions.X - 1 || 
-				YIndex == 0 || YIndex == VolumeDimensions.Y - 1 || 
+			if (MinDistance < 0 &&
+				(XIndex == 0 || XIndex == VolumeDimensions.X - 1 ||
+				YIndex == 0 || YIndex == VolumeDimensions.Y - 1 ||
 				ZIndex == 0 || ZIndex == VolumeDimensions.Z - 1))
 			{
 				bNegativeAtBorder = true;
@@ -605,7 +608,7 @@ void FMeshDistanceFieldAsyncTask::DoWork()
 
 void FMeshUtilities::GenerateSignedDistanceFieldVolumeData(
 	const FStaticMeshLODResources& LODModel,
-	class FQueuedThreadPool& ThreadPool,
+class FQueuedThreadPool& ThreadPool,
 	const TArray<EBlendMode>& MaterialBlendModes,
 	const FBoxSphereBounds& Bounds,
 	float DistanceFieldResolutionScale,
@@ -623,9 +626,9 @@ void FMeshUtilities::GenerateSignedDistanceFieldVolumeData(
 		float MaxDimension = FMath::Max(FMath::Max(BoundsSize.X, BoundsSize.Y), BoundsSize.Z);
 
 		// Consider the mesh a plane if it is very flat
-		const bool bMeshWasPlane = BoundsSize.Z * 100 < MaxDimension 
+		const bool bMeshWasPlane = BoundsSize.Z * 100 < MaxDimension
 			// And it lies mostly on the origin
-			&& Bounds.Origin.Z - Bounds.BoxExtent.Z < KINDA_SMALL_NUMBER 
+			&& Bounds.Origin.Z - Bounds.BoxExtent.Z < KINDA_SMALL_NUMBER
 			&& Bounds.Origin.Z + Bounds.BoxExtent.Z > -KINDA_SMALL_NUMBER;
 
 		for (int32 i = 0; i < Indices.Num(); i += 3)
@@ -659,7 +662,7 @@ void FMeshUtilities::GenerateSignedDistanceFieldVolumeData(
 						{
 							bTriangleIsOpaqueOrMasked = !IsTranslucentBlendMode(MaterialBlendModes[Section.MaterialIndex]);
 						}
-						
+
 						break;
 					}
 				}
@@ -667,7 +670,7 @@ void FMeshUtilities::GenerateSignedDistanceFieldVolumeData(
 				if (bTriangleIsOpaqueOrMasked)
 				{
 					BuildTriangles.Add(FkDOPBuildCollisionTriangle<uint32>(
-						bGenerateAsIfTwoSided, 
+						bGenerateAsIfTwoSided,
 						V0,
 						V1,
 						V2));
@@ -675,7 +678,7 @@ void FMeshUtilities::GenerateSignedDistanceFieldVolumeData(
 			}
 		}
 
-		TkDOPTree<const FMeshBuildDataProvider,uint32> kDopTree;
+		TkDOPTree<const FMeshBuildDataProvider, uint32> kDopTree;
 		kDopTree.Build(BuildTriangles);
 
 		//@todo - project setting
@@ -713,7 +716,7 @@ void FMeshUtilities::GenerateSignedDistanceFieldVolumeData(
 			const FVector DesiredDimensions(DistanceFieldVolumeBounds.GetSize() * FVector(NumVoxelsPerLocalSpaceUnit));
 
 			const FIntVector VolumeDimensions(
-				FMath::Clamp(FMath::TruncToInt(DesiredDimensions.X), MinNumVoxelsOneDim, MaxNumVoxelsOneDim), 
+				FMath::Clamp(FMath::TruncToInt(DesiredDimensions.X), MinNumVoxelsOneDim, MaxNumVoxelsOneDim),
 				FMath::Clamp(FMath::TruncToInt(DesiredDimensions.Y), MinNumVoxelsOneDim, MaxNumVoxelsOneDim),
 				FMath::Clamp(FMath::TruncToInt(DesiredDimensions.Z), MinNumVoxelsOneDim, MaxNumVoxelsOneDim));
 
@@ -752,11 +755,11 @@ void FMeshUtilities::GenerateSignedDistanceFieldVolumeData(
 			OutData.bBuiltAsIfTwoSided = bGenerateAsIfTwoSided;
 			OutData.bMeshWasPlane = bMeshWasPlane;
 
-			UE_LOG(LogMeshUtilities,Log,TEXT("Finished distance field build in %.1fs - %ux%ux%u distance field, %u triangles"), 
+			UE_LOG(LogMeshUtilities, Log, TEXT("Finished distance field build in %.1fs - %ux%ux%u distance field, %u triangles"),
 				(float)(FPlatformTime::Seconds() - StartTime),
-				VolumeDimensions.X, 
-				VolumeDimensions.Y, 
-				VolumeDimensions.Z, 
+				VolumeDimensions.X,
+				VolumeDimensions.Y,
+				VolumeDimensions.Z,
 				Indices.Num() / 3);
 
 			// Toss distance field if mesh was not closed
@@ -765,7 +768,7 @@ void FMeshUtilities::GenerateSignedDistanceFieldVolumeData(
 				OutData.Size = FIntVector(0, 0, 0);
 				OutData.DistanceFieldVolume.Empty();
 
-				UE_LOG(LogMeshUtilities,Log,TEXT("Discarded distance field as mesh was not closed!  Assign a two-sided material to fix."));
+				UE_LOG(LogMeshUtilities, Log, TEXT("Discarded distance field as mesh was not closed!  Assign a two-sided material to fix."));
 			}
 		}
 	}
@@ -775,7 +778,7 @@ void FMeshUtilities::GenerateSignedDistanceFieldVolumeData(
 
 void FMeshUtilities::GenerateSignedDistanceFieldVolumeData(
 	const FStaticMeshLODResources& LODModel,
-	class FQueuedThreadPool& ThreadPool,
+class FQueuedThreadPool& ThreadPool,
 	const TArray<EBlendMode>& MaterialBlendModes,
 	const FBoxSphereBounds& Bounds,
 	float DistanceFieldResolutionScale,
@@ -784,21 +787,21 @@ void FMeshUtilities::GenerateSignedDistanceFieldVolumeData(
 {
 	if (DistanceFieldResolutionScale > 0)
 	{
-		UE_LOG(LogMeshUtilities,Error,TEXT("Couldn't generate distance field for mesh, platform is missing required Vector intrinsics."));
+		UE_LOG(LogMeshUtilities, Error, TEXT("Couldn't generate distance field for mesh, platform is missing required Vector intrinsics."));
 	}
 }
 
 #endif
 
 /*------------------------------------------------------------------------------
-	NVTriStrip for cache optimizing index buffers.
+NVTriStrip for cache optimizing index buffers.
 ------------------------------------------------------------------------------*/
 
 namespace NvTriStrip
 {
 	/**
-	 * Converts 16 bit indices to 32 bit prior to passing them into the real GenerateStrips util method
-	 */
+	* Converts 16 bit indices to 32 bit prior to passing them into the real GenerateStrips util method
+	*/
 	void GenerateStrips(
 		const uint8* Indices,
 		bool Is32Bit,
@@ -827,15 +830,15 @@ namespace NvTriStrip
 	}
 
 	/**
-	 * Orders a triangle list for better vertex cache coherency.
-	 *
-	 * *** WARNING: This is safe to call for multiple threads IF AND ONLY IF all
-	 * threads call SetListsOnly(true) and SetCacheSize(CACHESIZE_GEFORCE3). If
-	 * NvTriStrip is ever used with different settings the library will need
-	 * some modifications to be thread-safe. ***
-	 */
+	* Orders a triangle list for better vertex cache coherency.
+	*
+	* *** WARNING: This is safe to call for multiple threads IF AND ONLY IF all
+	* threads call SetListsOnly(true) and SetCacheSize(CACHESIZE_GEFORCE3). If
+	* NvTriStrip is ever used with different settings the library will need
+	* some modifications to be thread-safe. ***
+	*/
 	template<typename IndexDataType, typename Allocator>
-	void CacheOptimizeIndexBuffer(TArray<IndexDataType,Allocator>& Indices)
+	void CacheOptimizeIndexBuffer(TArray<IndexDataType, Allocator>& Indices)
 	{
 		static_assert(sizeof(IndexDataType) == 2 || sizeof(IndexDataType) == 4, "Indices must be short or int.");
 
@@ -846,36 +849,36 @@ namespace NvTriStrip
 		SetListsOnly(true);
 		SetCacheSize(CACHESIZE_GEFORCE3);
 
-		GenerateStrips((uint8*)Indices.GetData(),Is32Bit,Indices.Num(),&PrimitiveGroups,&NumPrimitiveGroups);
+		GenerateStrips((uint8*)Indices.GetData(), Is32Bit, Indices.Num(), &PrimitiveGroups, &NumPrimitiveGroups);
 
 		Indices.Empty();
 		Indices.AddUninitialized(PrimitiveGroups->numIndices);
-	
-		if( Is32Bit )
+
+		if (Is32Bit)
 		{
-			FMemory::Memcpy(Indices.GetData(),PrimitiveGroups->indices,Indices.Num() * sizeof(IndexDataType));
+			FMemory::Memcpy(Indices.GetData(), PrimitiveGroups->indices, Indices.Num() * sizeof(IndexDataType));
 		}
 		else
 		{
-			for( uint32 I = 0; I < PrimitiveGroups->numIndices; ++I )
+			for (uint32 I = 0; I < PrimitiveGroups->numIndices; ++I)
 			{
 				Indices[I] = (uint16)PrimitiveGroups->indices[I];
 			}
 		}
 
-		delete [] PrimitiveGroups;
+		delete[] PrimitiveGroups;
 	}
 }
 
 /*------------------------------------------------------------------------------
-	Forsyth algorithm for cache optimizing index buffers.
+Forsyth algorithm for cache optimizing index buffers.
 ------------------------------------------------------------------------------*/
 
 namespace Forsyth
 {
 	/**
-	 * Converts 16 bit indices to 32 bit prior to passing them into the real OptimizeFaces util method
-	 */
+	* Converts 16 bit indices to 32 bit prior to passing them into the real OptimizeFaces util method
+	*/
 	void OptimizeFaces(
 		const uint8* Indices,
 		bool Is32Bit,
@@ -899,25 +902,25 @@ namespace Forsyth
 			{
 				NewIndices[Idx] = ((uint16*)Indices)[Idx];
 			}
-			OptimizeFaces(NewIndices.GetData(),NumIndices, NumVertices, OutIndices, CacheSize);
+			OptimizeFaces(NewIndices.GetData(), NumIndices, NumVertices, OutIndices, CacheSize);
 		}
 
 	}
 
 	/**
-	 * Orders a triangle list for better vertex cache coherency.
-	 */
+	* Orders a triangle list for better vertex cache coherency.
+	*/
 	template<typename IndexDataType, typename Allocator>
-	void CacheOptimizeIndexBuffer(TArray<IndexDataType,Allocator>& Indices)
+	void CacheOptimizeIndexBuffer(TArray<IndexDataType, Allocator>& Indices)
 	{
 		static_assert(sizeof(IndexDataType) == 2 || sizeof(IndexDataType) == 4, "Indices must be short or int.");
 		bool Is32Bit = sizeof(IndexDataType) == 4;
 
 		// Count the number of vertices
 		uint32 NumVertices = 0;
-		for(int32 Index = 0; Index < Indices.Num(); ++Index)
+		for (int32 Index = 0; Index < Indices.Num(); ++Index)
 		{
-			if(Indices[Index] > NumVertices)
+			if (Indices[Index] > NumVertices)
 			{
 				NumVertices = Indices[Index];
 			}
@@ -927,15 +930,15 @@ namespace Forsyth
 		TArray<uint32> OptimizedIndices;
 		OptimizedIndices.AddUninitialized(Indices.Num());
 		uint16 CacheSize = 32;
-		OptimizeFaces((uint8*)Indices.GetData(),Is32Bit,Indices.Num(),NumVertices, OptimizedIndices.GetData(), CacheSize);
+		OptimizeFaces((uint8*)Indices.GetData(), Is32Bit, Indices.Num(), NumVertices, OptimizedIndices.GetData(), CacheSize);
 
-		if( Is32Bit )
+		if (Is32Bit)
 		{
-			FMemory::Memcpy(Indices.GetData(),OptimizedIndices.GetData(),Indices.Num() * sizeof(IndexDataType));
+			FMemory::Memcpy(Indices.GetData(), OptimizedIndices.GetData(), Indices.Num() * sizeof(IndexDataType));
 		}
 		else
 		{
-			for( int32 I = 0; I < OptimizedIndices.Num(); ++I )
+			for (int32 I = 0; I < OptimizedIndices.Num(); ++I)
 			{
 				Indices[I] = (uint16)OptimizedIndices[I];
 			}
@@ -944,12 +947,12 @@ namespace Forsyth
 }
 
 void FMeshUtilities::CacheOptimizeIndexBuffer(TArray<uint16>& Indices)
-{	
-	if(bUsingNvTriStrip)
-	{				
-			NvTriStrip::CacheOptimizeIndexBuffer(Indices);
-		}
-	else if( !bDisableTriangleOrderOptimization )
+{
+	if (bUsingNvTriStrip)
+	{
+		NvTriStrip::CacheOptimizeIndexBuffer(Indices);
+	}
+	else if (!bDisableTriangleOrderOptimization)
 	{
 		Forsyth::CacheOptimizeIndexBuffer(Indices);
 	}
@@ -957,23 +960,23 @@ void FMeshUtilities::CacheOptimizeIndexBuffer(TArray<uint16>& Indices)
 
 void FMeshUtilities::CacheOptimizeIndexBuffer(TArray<uint32>& Indices)
 {
-	if(bUsingNvTriStrip)
+	if (bUsingNvTriStrip)
+	{
+		NvTriStrip::CacheOptimizeIndexBuffer(Indices);
+	}
+	else if (!bDisableTriangleOrderOptimization)
 	{
-			NvTriStrip::CacheOptimizeIndexBuffer(Indices);
-		}
-	else if( !bDisableTriangleOrderOptimization )
-	{		
 		Forsyth::CacheOptimizeIndexBuffer(Indices);
 	}
 }
 
 /*------------------------------------------------------------------------------
-	NVTessLib for computing adjacency used for tessellation.
+NVTessLib for computing adjacency used for tessellation.
 ------------------------------------------------------------------------------*/
 
 /**
- * Provides static mesh render data to the NVIDIA tessellation library.
- */
+* Provides static mesh render data to the NVIDIA tessellation library.
+*/
 class FStaticMeshNvRenderBuffer : public nv::RenderBuffer
 {
 public:
@@ -982,29 +985,29 @@ public:
 	FStaticMeshNvRenderBuffer(
 		const FPositionVertexBuffer& InPositionVertexBuffer,
 		const FStaticMeshVertexBuffer& InVertexBuffer,
-		const TArray<uint32>& Indices )
-		: PositionVertexBuffer( InPositionVertexBuffer )
-		, VertexBuffer( InVertexBuffer )
+		const TArray<uint32>& Indices)
+		: PositionVertexBuffer(InPositionVertexBuffer)
+		, VertexBuffer(InVertexBuffer)
 	{
-		check( PositionVertexBuffer.GetNumVertices() == VertexBuffer.GetNumVertices() );
-		mIb = new nv::IndexBuffer( (void*)Indices.GetData(), nv::IBT_U32, Indices.Num(), false );
+		check(PositionVertexBuffer.GetNumVertices() == VertexBuffer.GetNumVertices());
+		mIb = new nv::IndexBuffer((void*)Indices.GetData(), nv::IBT_U32, Indices.Num(), false);
 	}
 
 	/** Retrieve the position and first texture coordinate of the specified index. */
-	virtual nv::Vertex getVertex( unsigned int Index ) const
+	virtual nv::Vertex getVertex(unsigned int Index) const
 	{
 		nv::Vertex Vertex;
 
-		check( Index < PositionVertexBuffer.GetNumVertices() );
+		check(Index < PositionVertexBuffer.GetNumVertices());
 
-		const FVector& Position = PositionVertexBuffer.VertexPosition( Index );
+		const FVector& Position = PositionVertexBuffer.VertexPosition(Index);
 		Vertex.pos.x = Position.X;
 		Vertex.pos.y = Position.Y;
 		Vertex.pos.z = Position.Z;
 
-		if( VertexBuffer.GetNumTexCoords() )
+		if (VertexBuffer.GetNumTexCoords())
 		{
-			const FVector2D UV = VertexBuffer.GetVertexUV( Index, 0 );
+			const FVector2D UV = VertexBuffer.GetVertexUV(Index, 0);
 			Vertex.uv.x = UV.X;
 			Vertex.uv.y = UV.Y;
 		}
@@ -1026,42 +1029,42 @@ private:
 	const FStaticMeshVertexBuffer& VertexBuffer;
 
 	/** Copying is forbidden. */
-	FStaticMeshNvRenderBuffer( const FStaticMeshNvRenderBuffer& ); 
-	FStaticMeshNvRenderBuffer& operator=( const FStaticMeshNvRenderBuffer& );
+	FStaticMeshNvRenderBuffer(const FStaticMeshNvRenderBuffer&);
+	FStaticMeshNvRenderBuffer& operator=(const FStaticMeshNvRenderBuffer&);
 };
 
 /**
- * Provides skeletal mesh render data to the NVIDIA tessellation library.
- */
+* Provides skeletal mesh render data to the NVIDIA tessellation library.
+*/
 class FSkeletalMeshNvRenderBuffer : public nv::RenderBuffer
 {
 public:
 
 	/** Construct from static mesh render buffers. */
-	FSkeletalMeshNvRenderBuffer( 
+	FSkeletalMeshNvRenderBuffer(
 		const TArray<FSoftSkinVertex>& InVertexBuffer,
 		const uint32 InTexCoordCount,
-		const TArray<uint32>& Indices )
-		: VertexBuffer( InVertexBuffer )
-		, TexCoordCount( InTexCoordCount )
+		const TArray<uint32>& Indices)
+		: VertexBuffer(InVertexBuffer)
+		, TexCoordCount(InTexCoordCount)
 	{
-		mIb = new nv::IndexBuffer( (void*)Indices.GetData(), nv::IBT_U32, Indices.Num(), false );
+		mIb = new nv::IndexBuffer((void*)Indices.GetData(), nv::IBT_U32, Indices.Num(), false);
 	}
 
 	/** Retrieve the position and first texture coordinate of the specified index. */
-	virtual nv::Vertex getVertex( unsigned int Index ) const
+	virtual nv::Vertex getVertex(unsigned int Index) const
 	{
 		nv::Vertex Vertex;
 
-		check( Index < (unsigned int)VertexBuffer.Num() );
+		check(Index < (unsigned int)VertexBuffer.Num());
 
-		const FSoftSkinVertex& SrcVertex = VertexBuffer[ Index ];
+		const FSoftSkinVertex& SrcVertex = VertexBuffer[Index];
 
 		Vertex.pos.x = SrcVertex.Position.X;
 		Vertex.pos.y = SrcVertex.Position.Y;
 		Vertex.pos.z = SrcVertex.Position.Z;
 
-		if( TexCoordCount > 0 )
+		if (TexCoordCount > 0)
 		{
 			Vertex.uv.x = SrcVertex.UVs[0].X;
 			Vertex.uv.y = SrcVertex.UVs[0].Y;
@@ -1081,7 +1084,7 @@ private:
 	const uint32 TexCoordCount;
 
 	/** Copying is forbidden. */
-	FSkeletalMeshNvRenderBuffer(const FSkeletalMeshNvRenderBuffer&); 
+	FSkeletalMeshNvRenderBuffer(const FSkeletalMeshNvRenderBuffer&);
 	FSkeletalMeshNvRenderBuffer& operator=(const FSkeletalMeshNvRenderBuffer&);
 };
 
@@ -1092,17 +1095,17 @@ static void BuildStaticAdjacencyIndexBuffer(
 	TArray<uint32>& OutPnAenIndices
 	)
 {
-	if ( Indices.Num() )
+	if (Indices.Num())
 	{
-		FStaticMeshNvRenderBuffer StaticMeshRenderBuffer( PositionVertexBuffer, VertexBuffer, Indices );
-		nv::IndexBuffer* PnAENIndexBuffer = nv::tess::buildTessellationBuffer( &StaticMeshRenderBuffer, nv::DBM_PnAenDominantCorner, true );
-		check( PnAENIndexBuffer );
+		FStaticMeshNvRenderBuffer StaticMeshRenderBuffer(PositionVertexBuffer, VertexBuffer, Indices);
+		nv::IndexBuffer* PnAENIndexBuffer = nv::tess::buildTessellationBuffer(&StaticMeshRenderBuffer, nv::DBM_PnAenDominantCorner, true);
+		check(PnAENIndexBuffer);
 		const int32 IndexCount = (int32)PnAENIndexBuffer->getLength();
-		OutPnAenIndices.Empty( IndexCount );
-		OutPnAenIndices.AddUninitialized( IndexCount );
-		for ( int32 Index = 0; Index < IndexCount; ++Index )
+		OutPnAenIndices.Empty(IndexCount);
+		OutPnAenIndices.AddUninitialized(IndexCount);
+		for (int32 Index = 0; Index < IndexCount; ++Index)
 		{
-			OutPnAenIndices[ Index ] = (*PnAENIndexBuffer)[Index];
+			OutPnAenIndices[Index] = (*PnAENIndexBuffer)[Index];
 		}
 		delete PnAENIndexBuffer;
 	}
@@ -1119,17 +1122,17 @@ void FMeshUtilities::BuildSkeletalAdjacencyIndexBuffer(
 	TArray<uint32>& OutPnAenIndices
 	)
 {
-	if ( Indices.Num() )
+	if (Indices.Num())
 	{
-		FSkeletalMeshNvRenderBuffer SkeletalMeshRenderBuffer( VertexBuffer, TexCoordCount, Indices );
-		nv::IndexBuffer* PnAENIndexBuffer = nv::tess::buildTessellationBuffer( &SkeletalMeshRenderBuffer, nv::DBM_PnAenDominantCorner, true );
-		check( PnAENIndexBuffer );
+		FSkeletalMeshNvRenderBuffer SkeletalMeshRenderBuffer(VertexBuffer, TexCoordCount, Indices);
+		nv::IndexBuffer* PnAENIndexBuffer = nv::tess::buildTessellationBuffer(&SkeletalMeshRenderBuffer, nv::DBM_PnAenDominantCorner, true);
+		check(PnAENIndexBuffer);
 		const int32 IndexCount = (int32)PnAENIndexBuffer->getLength();
-		OutPnAenIndices.Empty( IndexCount );
-		OutPnAenIndices.AddUninitialized( IndexCount );
-		for ( int32 Index = 0; Index < IndexCount; ++Index )
+		OutPnAenIndices.Empty(IndexCount);
+		OutPnAenIndices.AddUninitialized(IndexCount);
+		for (int32 Index = 0; Index < IndexCount; ++Index)
 		{
-			OutPnAenIndices[ Index ] = (*PnAENIndexBuffer)[Index];
+			OutPnAenIndices[Index] = (*PnAENIndexBuffer)[Index];
 		}
 		delete PnAENIndexBuffer;
 	}
@@ -1152,8 +1155,8 @@ void FMeshUtilities::RechunkSkeletalMeshModels(USkeletalMesh* SrcMesh, int32 Max
 
 	for (int32 ModelIndex = 0; ModelIndex < SrcMeshResource->LODModels.Num(); ++ModelIndex)
 	{
-		FSkinnedModelData& TmpModelData = *new(ModelData) FSkinnedModelData();
-		SkeletalMeshTools::CopySkinnedModelData(TmpModelData,SrcMeshResource->LODModels[ModelIndex]);
+		FSkinnedModelData& TmpModelData = *new(ModelData)FSkinnedModelData();
+		SkeletalMeshTools::CopySkinnedModelData(TmpModelData, SrcMeshResource->LODModels[ModelIndex]);
 	}
 
 	for (int32 ModelIndex = 0; ModelIndex < ModelData.Num(); ++ModelIndex)
@@ -1163,10 +1166,10 @@ void FMeshUtilities::RechunkSkeletalMeshModels(USkeletalMesh* SrcMesh, int32 Max
 		TArray<ETriangleSortOption> SectionSortOptions;
 
 		const FSkinnedModelData& SrcModel = ModelData[ModelIndex];
-		FStaticLODModel& DestModel = *new(DestModels) FStaticLODModel();
+		FStaticLODModel& DestModel = *new(DestModels)FStaticLODModel();
 
-		SkeletalMeshTools::UnchunkSkeletalModel(Chunks,PointToOriginalMap,SrcModel);
-		SkeletalMeshTools::ChunkSkinnedVertices(Chunks,MaxBonesPerChunk);
+		SkeletalMeshTools::UnchunkSkeletalModel(Chunks, PointToOriginalMap, SrcModel);
+		SkeletalMeshTools::ChunkSkinnedVertices(Chunks, MaxBonesPerChunk);
 
 		for (int32 ChunkIndex = 0; ChunkIndex < Chunks.Num(); ++ChunkIndex)
 		{
@@ -1175,7 +1178,7 @@ void FMeshUtilities::RechunkSkeletalMeshModels(USkeletalMesh* SrcMesh, int32 Max
 		}
 		check(SectionSortOptions.Num() == Chunks.Num());
 
-		BuildSkeletalModelFromChunks(DestModel,RefSkeleton,Chunks,PointToOriginalMap);
+		BuildSkeletalModelFromChunks(DestModel, RefSkeleton, Chunks, PointToOriginalMap);
 		check(DestModel.Sections.Num() == DestModel.Chunks.Num());
 		check(DestModel.Sections.Num() == SectionSortOptions.Num());
 
@@ -1183,7 +1186,7 @@ void FMeshUtilities::RechunkSkeletalMeshModels(USkeletalMesh* SrcMesh, int32 Max
 		DestModel.BuildVertexBuffers(VertexBufferBuildFlags);
 		for (int32 SectionIndex = 0; SectionIndex < DestModel.Sections.Num(); ++SectionIndex)
 		{
-			DestModel.SortTriangles(TriangleSortCenter,bHaveTriangleSortCenter,SectionIndex,SectionSortOptions[SectionIndex]);
+			DestModel.SortTriangles(TriangleSortCenter, bHaveTriangleSortCenter, SectionIndex, SectionSortOptions[SectionIndex]);
 		}
 	}
 
@@ -1196,18 +1199,18 @@ void FMeshUtilities::RechunkSkeletalMeshModels(USkeletalMesh* SrcMesh, int32 Max
 
 void FMeshUtilities::CalcBoneVertInfos(USkeletalMesh* SkeletalMesh, TArray<FBoneVertInfo>& Infos, bool bOnlyDominant)
 {
-	SkeletalMeshTools::CalcBoneVertInfos(SkeletalMesh,Infos,bOnlyDominant);
+	SkeletalMeshTools::CalcBoneVertInfos(SkeletalMesh, Infos, bOnlyDominant);
 }
 
 /**
- * Builds a renderable skeletal mesh LOD model. Note that the array of chunks
- * will be destroyed during this process!
- * @param LODModel				Upon return contains a renderable skeletal mesh LOD model.
- * @param RefSkeleton			The reference skeleton associated with the model.
- * @param Chunks				Skinned mesh chunks from which to build the renderable model.
- * @param PointToOriginalMap	Maps a vertex's RawPointIdx to its index at import time.
- */
-void FMeshUtilities::BuildSkeletalModelFromChunks(FStaticLODModel& LODModel,const FReferenceSkeleton& RefSkeleton,TArray<FSkinnedMeshChunk*>& Chunks,const TArray<int32>& PointToOriginalMap)
+* Builds a renderable skeletal mesh LOD model. Note that the array of chunks
+* will be destroyed during this process!
+* @param LODModel				Upon return contains a renderable skeletal mesh LOD model.
+* @param RefSkeleton			The reference skeleton associated with the model.
+* @param Chunks				Skinned mesh chunks from which to build the renderable model.
+* @param PointToOriginalMap	Maps a vertex's RawPointIdx to its index at import time.
+*/
+void FMeshUtilities::BuildSkeletalModelFromChunks(FStaticLODModel& LODModel, const FReferenceSkeleton& RefSkeleton, TArray<FSkinnedMeshChunk*>& Chunks, const TArray<int32>& PointToOriginalMap)
 {
 #if WITH_EDITORONLY_DATA
 	// Clear out any data currently held in the LOD model.
@@ -1229,7 +1232,7 @@ void FMeshUtilities::BuildSkeletalModelFromChunks(FStaticLODModel& LODModel,cons
 		Section.ChunkIndex = ChunkIndex;
 
 		FSkelMeshChunk& Chunk = *new(LODModel.Chunks) FSkelMeshChunk();
-		Exchange(Chunk.BoneMap,SrcChunk->BoneMap);
+		Exchange(Chunk.BoneMap, SrcChunk->BoneMap);
 
 		// Update the active bone indices on the LOD model.
 		for (int32 BoneIndex = 0; BoneIndex < Chunk.BoneMap.Num(); ++BoneIndex)
@@ -1237,7 +1240,7 @@ void FMeshUtilities::BuildSkeletalModelFromChunks(FStaticLODModel& LODModel,cons
 			LODModel.ActiveBoneIndices.AddUnique(Chunk.BoneMap[BoneIndex]);
 		}
 	}
-	
+
 	LODModel.ActiveBoneIndices.Sort();
 
 	// Reset 'final vertex to import vertex' map info
@@ -1248,7 +1251,7 @@ void FMeshUtilities::BuildSkeletalModelFromChunks(FStaticLODModel& LODModel,cons
 	TArray< TArray<uint32> > VertexIndexRemap;
 	VertexIndexRemap.Empty(LODModel.Sections.Num());
 	// Pack the chunk vertices into a single vertex buffer.
-	TArray<uint32> RawPointIndices;	
+	TArray<uint32> RawPointIndices;
 	LODModel.NumVertices = 0;
 
 	int32 PrevMaterialIndex = -1;
@@ -1275,7 +1278,7 @@ void FMeshUtilities::BuildSkeletalModelFromChunks(FStaticLODModel& LODModel,cons
 		// Calculate the number of triangles in the section.  Note that CacheOptimize may change the number of triangles in the index buffer!
 		Section.NumTriangles = ChunkIndices.Num() / 3;
 		TArray<FSoftSkinBuildVertex> OriginalVertices;
-		Exchange(ChunkVertices,OriginalVertices);
+		Exchange(ChunkVertices, OriginalVertices);
 		ChunkVertices.AddUninitialized(OriginalVertices.Num());
 
 		TArray<int32> IndexCache;
@@ -1307,17 +1310,17 @@ void FMeshUtilities::BuildSkeletalModelFromChunks(FStaticLODModel& LODModel,cons
 	}
 
 	// Build the arrays of rigid and soft vertices on the model's chunks.
-	for(int32 SectionIndex = 0;SectionIndex < LODModel.Sections.Num();SectionIndex++)
+	for (int32 SectionIndex = 0; SectionIndex < LODModel.Sections.Num(); SectionIndex++)
 	{
 		FSkelMeshSection& Section = LODModel.Sections[SectionIndex];
 		int32 ChunkIndex = Section.ChunkIndex;
-		FSkelMeshChunk& Chunk = LODModel.Chunks[ChunkIndex]; 
+		FSkelMeshChunk& Chunk = LODModel.Chunks[ChunkIndex];
 		TArray<FSoftSkinBuildVertex>& ChunkVertices = Chunks[ChunkIndex]->Vertices;
 
-		if( IsInGameThread() )
+		if (IsInGameThread())
 		{
 			// Only update status if in the game thread.  When importing morph targets, this function can run in another thread
-			GWarn->StatusUpdate( ChunkIndex, LODModel.Chunks.Num(), NSLOCTEXT("UnrealEd", "ProcessingChunks", "Processing Chunks") );
+			GWarn->StatusUpdate(ChunkIndex, LODModel.Chunks.Num(), NSLOCTEXT("UnrealEd", "ProcessingChunks", "Processing Chunks"));
 		}
 
 		CurrentVertexIndex = 0;
@@ -1331,45 +1334,45 @@ void FMeshUtilities::BuildSkeletalModelFromChunks(FStaticLODModel& LODModel,cons
 		LODModel.NumVertices += ChunkVertices.Num();
 
 		// Separate the section's vertices into rigid and soft vertices.
-		TArray<uint32>& ChunkVertexIndexRemap = *new(VertexIndexRemap) TArray<uint32>();
+		TArray<uint32>& ChunkVertexIndexRemap = *new(VertexIndexRemap)TArray<uint32>();
 		ChunkVertexIndexRemap.AddUninitialized(ChunkVertices.Num());
-		for(int32 VertexIndex = 0;VertexIndex < ChunkVertices.Num();VertexIndex++)
+		for (int32 VertexIndex = 0; VertexIndex < ChunkVertices.Num(); VertexIndex++)
 		{
 			const FSoftSkinBuildVertex& SoftVertex = ChunkVertices[VertexIndex];
-			if(SoftVertex.InfluenceWeights[1] == 0)
+			if (SoftVertex.InfluenceWeights[1] == 0)
 			{
 				FRigidSkinVertex RigidVertex;
 				RigidVertex.Position = SoftVertex.Position;
 				RigidVertex.TangentX = SoftVertex.TangentX;
 				RigidVertex.TangentY = SoftVertex.TangentY;
 				RigidVertex.TangentZ = SoftVertex.TangentZ;
-				FMemory::Memcpy( RigidVertex.UVs, SoftVertex.UVs, sizeof(FVector2D)*MAX_TEXCOORDS );
+				FMemory::Memcpy(RigidVertex.UVs, SoftVertex.UVs, sizeof(FVector2D)*MAX_TEXCOORDS);
 				RigidVertex.Color = SoftVertex.Color;
 				RigidVertex.Bone = SoftVertex.InfluenceBones[0];
 				// make sure it exists in bone map
-				check (Chunk.BoneMap.IsValidIndex(SoftVertex.InfluenceBones[0]));
+				check(Chunk.BoneMap.IsValidIndex(SoftVertex.InfluenceBones[0]));
 				Chunk.RigidVertices.Add(RigidVertex);
 				ChunkVertexIndexRemap[VertexIndex] = (uint32)(Chunk.BaseVertexIndex + CurrentVertexIndex);
 				CurrentVertexIndex++;
 				// add the index to the original wedge point source of this vertex
-				RawPointIndices.Add( SoftVertex.PointWedgeIdx );
+				RawPointIndices.Add(SoftVertex.PointWedgeIdx);
 				// Also remember import index
 				const int32 RawVertIndex = PointToOriginalMap[SoftVertex.PointWedgeIdx];
 				LODModel.MeshToImportVertexMap.Add(RawVertIndex);
 				LODModel.MaxImportVertex = FMath::Max<float>(LODModel.MaxImportVertex, RawVertIndex);
 			}
 		}
-		for(int32 VertexIndex = 0;VertexIndex < ChunkVertices.Num();VertexIndex++)
+		for (int32 VertexIndex = 0; VertexIndex < ChunkVertices.Num(); VertexIndex++)
 		{
 			const FSoftSkinBuildVertex& SoftVertex = ChunkVertices[VertexIndex];
-			if(SoftVertex.InfluenceWeights[1] > 0)
+			if (SoftVertex.InfluenceWeights[1] > 0)
 			{
 				FSoftSkinVertex NewVertex;
 				NewVertex.Position = SoftVertex.Position;
 				NewVertex.TangentX = SoftVertex.TangentX;
 				NewVertex.TangentY = SoftVertex.TangentY;
 				NewVertex.TangentZ = SoftVertex.TangentZ;
-				FMemory::Memcpy( NewVertex.UVs, SoftVertex.UVs, sizeof(FVector2D)*MAX_TEXCOORDS );
+				FMemory::Memcpy(NewVertex.UVs, SoftVertex.UVs, sizeof(FVector2D)*MAX_TEXCOORDS);
 				NewVertex.Color = SoftVertex.Color;
 				for (int32 i = 0; i < MAX_TOTAL_INFLUENCES; ++i)
 				{
@@ -1386,7 +1389,7 @@ void FMeshUtilities::BuildSkeletalModelFromChunks(FStaticLODModel& LODModel,cons
 				ChunkVertexIndexRemap[VertexIndex] = (uint32)(Chunk.BaseVertexIndex + CurrentVertexIndex);
 				CurrentVertexIndex++;
 				// add the index to the original wedge point source of this vertex
-				RawPointIndices.Add( SoftVertex.PointWedgeIdx );
+				RawPointIndices.Add(SoftVertex.PointWedgeIdx);
 				// Also remember import index
 				const int32 RawVertIndex = PointToOriginalMap[SoftVertex.PointWedgeIdx];
 				LODModel.MeshToImportVertexMap.Add(RawVertIndex);
@@ -1412,18 +1415,18 @@ void FMeshUtilities::BuildSkeletalModelFromChunks(FStaticLODModel& LODModel,cons
 
 	// Copy raw point indices to LOD model.
 	LODModel.RawPointIndices.RemoveBulkData();
-	if( RawPointIndices.Num() )
+	if (RawPointIndices.Num())
 	{
 		LODModel.RawPointIndices.Lock(LOCK_READ_WRITE);
-		void* Dest = LODModel.RawPointIndices.Realloc( RawPointIndices.Num() );
-		FMemory::Memcpy( Dest, RawPointIndices.GetData(), LODModel.RawPointIndices.GetBulkDataSize() );
+		void* Dest = LODModel.RawPointIndices.Realloc(RawPointIndices.Num());
+		FMemory::Memcpy(Dest, RawPointIndices.GetData(), LODModel.RawPointIndices.GetBulkDataSize());
 		LODModel.RawPointIndices.Unlock();
 	}
 
 #if DISALLOW_32BIT_INDICES
 	LODModel.MultiSizeIndexContainer.CreateIndexBuffer(sizeof(uint16));
 #else
-	LODModel.MultiSizeIndexContainer.CreateIndexBuffer((LODModel.NumVertices < MAX_uint16)? sizeof(uint16): sizeof(uint32));
+	LODModel.MultiSizeIndexContainer.CreateIndexBuffer((LODModel.NumVertices < MAX_uint16) ? sizeof(uint16) : sizeof(uint32));
 #endif
 
 	// Finish building the sections.
@@ -1454,25 +1457,25 @@ void FMeshUtilities::BuildSkeletalModelFromChunks(FStaticLODModel& LODModel,cons
 	// Build the adjacency index buffer used for tessellation.
 	{
 		TArray<FSoftSkinVertex> Vertices;
-		LODModel.GetVertices( Vertices );
+		LODModel.GetVertices(Vertices);
 
 		FMultiSizeIndexContainerData IndexData;
-		LODModel.MultiSizeIndexContainer.GetIndexBufferData( IndexData );
+		LODModel.MultiSizeIndexContainer.GetIndexBufferData(IndexData);
 
 		FMultiSizeIndexContainerData AdjacencyIndexData;
 		AdjacencyIndexData.DataTypeSize = IndexData.DataTypeSize;
 
-		BuildSkeletalAdjacencyIndexBuffer( Vertices, LODModel.NumTexCoords, IndexData.Indices, AdjacencyIndexData.Indices );
-		LODModel.AdjacencyMultiSizeIndexContainer.RebuildIndexBuffer( AdjacencyIndexData );
+		BuildSkeletalAdjacencyIndexBuffer(Vertices, LODModel.NumTexCoords, IndexData.Indices, AdjacencyIndexData.Indices);
+		LODModel.AdjacencyMultiSizeIndexContainer.RebuildIndexBuffer(AdjacencyIndexData);
 	}
 
 	// Compute the required bones for this model.
-	USkeletalMesh::CalculateRequiredBones(LODModel,RefSkeleton,NULL);
+	USkeletalMesh::CalculateRequiredBones(LODModel, RefSkeleton, NULL);
 #endif // #if WITH_EDITORONLY_DATA
 }
 
 /*------------------------------------------------------------------------------
-	Common functionality.
+Common functionality.
 ------------------------------------------------------------------------------*/
 
 /** Helper struct for building acceleration structures. */
@@ -1514,9 +1517,9 @@ static int32 ComputeNumTexCoords(FRawMesh const& RawMesh, int32 MaxSupportedTexC
 }
 
 /**
- * Returns true if the specified points are about equal
- */
-inline bool PointsEqual(const FVector& V1,const FVector& V2, float ComparisonThreshold)
+* Returns true if the specified points are about equal
+*/
+inline bool PointsEqual(const FVector& V1, const FVector& V2, float ComparisonThreshold)
 {
 	if (FMath::Abs(V1.X - V2.X) > ComparisonThreshold
 		|| FMath::Abs(V1.Y - V2.Y) > ComparisonThreshold
@@ -1527,12 +1530,12 @@ inline bool PointsEqual(const FVector& V1,const FVector& V2, float ComparisonThr
 	return true;
 }
 
-inline bool UVsEqual(const FVector2D& UV1,const FVector2D& UV2)
+inline bool UVsEqual(const FVector2D& UV1, const FVector2D& UV2)
 {
-	if(FMath::Abs(UV1.X - UV2.X) > (1.0f / 1024.0f))
+	if (FMath::Abs(UV1.X - UV2.X) > (1.0f / 1024.0f))
 		return false;
 
-	if(FMath::Abs(UV1.Y - UV2.Y) > (1.0f / 1024.0f))
+	if (FMath::Abs(UV1.Y - UV2.Y) > (1.0f / 1024.0f))
 		return false;
 
 	return true;
@@ -1551,66 +1554,66 @@ struct FMeshEdge
 };
 
 /**
- * This helper class builds the edge list for a mesh. It uses a hash of vertex
- * positions to edges sharing that vertex to remove the n^2 searching of all
- * previously added edges. This class is templatized so it can be used with
- * either static mesh or skeletal mesh vertices
- */
+* This helper class builds the edge list for a mesh. It uses a hash of vertex
+* positions to edges sharing that vertex to remove the n^2 searching of all
+* previously added edges. This class is templatized so it can be used with
+* either static mesh or skeletal mesh vertices
+*/
 template <class VertexClass> class TEdgeBuilder
 {
 protected:
 	/**
-	 * The list of indices to build the edge data from
-	 */
+	* The list of indices to build the edge data from
+	*/
 	const TArray<uint32>& Indices;
 	/**
-	 * The array of verts for vertex position comparison
-	 */
+	* The array of verts for vertex position comparison
+	*/
 	const TArray<VertexClass>& Vertices;
 	/**
-	 * The array of edges to create
-	 */
+	* The array of edges to create
+	*/
 	TArray<FMeshEdge>& Edges;
 	/**
-	 * List of edges that start with a given vertex
-	 */
-	TMultiMap<FVector,FMeshEdge*> VertexToEdgeList;
+	* List of edges that start with a given vertex
+	*/
+	TMultiMap<FVector, FMeshEdge*> VertexToEdgeList;
 
 	/**
-	 * This function determines whether a given edge matches or not. It must be
-	 * provided by derived classes since they have the specific information that
-	 * this class doesn't know about (vertex info, influences, etc)
-	 *
-	 * @param Index1 The first index of the edge being checked
-	 * @param Index2 The second index of the edge
-	 * @param OtherEdge The edge to compare. Was found via the map
-	 *
-	 * @return true if the edge is a match, false otherwise
-	 */
-	virtual bool DoesEdgeMatch(int32 Index1,int32 Index2,FMeshEdge* OtherEdge) = 0;
+	* This function determines whether a given edge matches or not. It must be
+	* provided by derived classes since they have the specific information that
+	* this class doesn't know about (vertex info, influences, etc)
+	*
+	* @param Index1 The first index of the edge being checked
+	* @param Index2 The second index of the edge
+	* @param OtherEdge The edge to compare. Was found via the map
+	*
+	* @return true if the edge is a match, false otherwise
+	*/
+	virtual bool DoesEdgeMatch(int32 Index1, int32 Index2, FMeshEdge* OtherEdge) = 0;
 
 	/**
-	 * Searches the list of edges to see if this one matches an existing and
-	 * returns a pointer to it if it does
-	 *
-	 * @param Index1 the first index to check for
-	 * @param Index2 the second index to check for
-	 *
-	 * @return NULL if no edge was found, otherwise the edge that was found
-	 */
-	inline FMeshEdge* FindOppositeEdge(int32 Index1,int32 Index2)
+	* Searches the list of edges to see if this one matches an existing and
+	* returns a pointer to it if it does
+	*
+	* @param Index1 the first index to check for
+	* @param Index2 the second index to check for
+	*
+	* @return NULL if no edge was found, otherwise the edge that was found
+	*/
+	inline FMeshEdge* FindOppositeEdge(int32 Index1, int32 Index2)
 	{
 		FMeshEdge* Edge = NULL;
 		TArray<FMeshEdge*> EdgeList;
 		// Search the hash for a corresponding vertex
-		VertexToEdgeList.MultiFind(Vertices[Index2].Position,EdgeList);
+		VertexToEdgeList.MultiFind(Vertices[Index2].Position, EdgeList);
 		// Now search through the array for a match or not
 		for (int32 EdgeIndex = 0; EdgeIndex < EdgeList.Num() && Edge == NULL;
 			EdgeIndex++)
 		{
 			FMeshEdge* OtherEdge = EdgeList[EdgeIndex];
 			// See if this edge matches the passed in edge
-			if (OtherEdge != NULL && DoesEdgeMatch(Index1,Index2,OtherEdge))
+			if (OtherEdge != NULL && DoesEdgeMatch(Index1, Index2, OtherEdge))
 			{
 				// We have a match
 				Edge = OtherEdge;
@@ -1620,17 +1623,17 @@ protected:
 	}
 
 	/**
-	 * Updates an existing edge if found or adds the new edge to the list
-	 *
-	 * @param Index1 the first index in the edge
-	 * @param Index2 the second index in the edge
-	 * @param Triangle the triangle that this edge was found in
-	 */
-	inline void AddEdge(int32 Index1,int32 Index2,int32 Triangle)
+	* Updates an existing edge if found or adds the new edge to the list
+	*
+	* @param Index1 the first index in the edge
+	* @param Index2 the second index in the edge
+	* @param Triangle the triangle that this edge was found in
+	*/
+	inline void AddEdge(int32 Index1, int32 Index2, int32 Triangle)
 	{
 		// If this edge matches another then just fill the other triangle
 		// otherwise add it
-		FMeshEdge* OtherEdge = FindOppositeEdge(Index1,Index2);
+		FMeshEdge* OtherEdge = FindOppositeEdge(Index1, Index2);
 		if (OtherEdge == NULL)
 		{
 			// Add a new edge to the array
@@ -1641,7 +1644,7 @@ protected:
 			Edges[EdgeIndex].Faces[1] = -1;
 			// Also add this edge to the hash for faster searches
 			// NOTE: This relies on the array never being realloced!
-			VertexToEdgeList.Add(Vertices[Index1].Position,&Edges[EdgeIndex]);
+			VertexToEdgeList.Add(Vertices[Index1].Position, &Edges[EdgeIndex]);
 		}
 		else
 		{
@@ -1651,8 +1654,8 @@ protected:
 
 public:
 	/**
-	 * Initializes the values for the code that will build the mesh edge list
-	 */
+	* Initializes the values for the code that will build the mesh edge list
+	*/
 	TEdgeBuilder(const TArray<uint32>& InIndices,
 		const TArray<VertexClass>& InVertices,
 		TArray<FMeshEdge>& OutEdges) :
@@ -1670,9 +1673,9 @@ public:
 
 
 	/**
-	 * Uses a hash of indices to edge lists so that it can avoid the n^2 search
-	 * through the full edge list
-	 */
+	* Uses a hash of indices to edge lists so that it can avoid the n^2 search
+	* through the full edge list
+	*/
 	void FindEdges(void)
 	{
 		// @todo Handle something other than trilists when building edges
@@ -1688,41 +1691,41 @@ public:
 			int32 Index2 = Indices[TriangleIndex + 1];
 			int32 Index3 = Indices[TriangleIndex + 2];
 			// Add the first to second edge
-			AddEdge(Index1,Index2,Triangle);
+			AddEdge(Index1, Index2, Triangle);
 			// Now add the second to third
-			AddEdge(Index2,Index3,Triangle);
+			AddEdge(Index2, Index3, Triangle);
 			// Add the third to first edge
-			AddEdge(Index3,Index1,Triangle);
+			AddEdge(Index3, Index1, Triangle);
 		}
 	}
 };
 
 /**
- * This is the static mesh specific version for finding edges
- */
+* This is the static mesh specific version for finding edges
+*/
 class FStaticMeshEdgeBuilder : public TEdgeBuilder<FStaticMeshBuildVertex>
 {
 public:
 	/**
-	 * Constructor that passes all work to the parent class
-	 */
+	* Constructor that passes all work to the parent class
+	*/
 	FStaticMeshEdgeBuilder(const TArray<uint32>& InIndices,
 		const TArray<FStaticMeshBuildVertex>& InVertices,
 		TArray<FMeshEdge>& OutEdges) :
-		TEdgeBuilder<FStaticMeshBuildVertex>(InIndices,InVertices,OutEdges)
+		TEdgeBuilder<FStaticMeshBuildVertex>(InIndices, InVertices, OutEdges)
 	{
 	}
 
 	/**
-	 * This function determines whether a given edge matches or not for a static mesh
-	 *
-	 * @param Index1 The first index of the edge being checked
-	 * @param Index2 The second index of the edge
-	 * @param OtherEdge The edge to compare. Was found via the map
-	 *
-	 * @return true if the edge is a match, false otherwise
-	 */
-	bool DoesEdgeMatch(int32 Index1,int32 Index2,FMeshEdge* OtherEdge)
+	* This function determines whether a given edge matches or not for a static mesh
+	*
+	* @param Index1 The first index of the edge being checked
+	* @param Index2 The second index of the edge
+	* @param OtherEdge The edge to compare. Was found via the map
+	*
+	* @return true if the edge is a match, false otherwise
+	*/
+	bool DoesEdgeMatch(int32 Index1, int32 Index2, FMeshEdge* OtherEdge)
 	{
 		return Vertices[OtherEdge->Vertices[1]].Position == Vertices[Index1].Position &&
 			OtherEdge->Faces[1] == -1;
@@ -1742,7 +1745,7 @@ static void ComputeTriangleTangents(
 	TriangleTangentY.Empty(NumTriangles);
 	TriangleTangentZ.Empty(NumTriangles);
 
-	for (int32 TriangleIndex = 0;TriangleIndex < NumTriangles; TriangleIndex++)
+	for (int32 TriangleIndex = 0; TriangleIndex < NumTriangles; TriangleIndex++)
 	{
 		int32 UVIndex = 0;
 
@@ -1752,29 +1755,29 @@ static void ComputeTriangleTangents(
 			P[i] = GetPositionForWedge(RawMesh, TriangleIndex * 3 + i);
 		}
 
-		const FVector Normal = ((P[1] - P[2])^(P[0] - P[2])).GetSafeNormal(ComparisonThreshold);
+		const FVector Normal = ((P[1] - P[2]) ^ (P[0] - P[2])).GetSafeNormal(ComparisonThreshold);
 		FMatrix	ParameterToLocal(
 			FPlane(P[1].X - P[0].X, P[1].Y - P[0].Y, P[1].Z - P[0].Z, 0),
 			FPlane(P[2].X - P[0].X, P[2].Y - P[0].Y, P[2].Z - P[0].Z, 0),
-			FPlane(P[0].X,          P[0].Y,          P[0].Z,          0),
-			FPlane(0,               0,               0,				  1)
+			FPlane(P[0].X, P[0].Y, P[0].Z, 0),
+			FPlane(0, 0, 0, 1)
 			);
 
 		FVector2D T1 = RawMesh.WedgeTexCoords[UVIndex][TriangleIndex * 3 + 0];
 		FVector2D T2 = RawMesh.WedgeTexCoords[UVIndex][TriangleIndex * 3 + 1];
 		FVector2D T3 = RawMesh.WedgeTexCoords[UVIndex][TriangleIndex * 3 + 2];
 		FMatrix ParameterToTexture(
-			FPlane(	T2.X - T1.X,	T2.Y - T1.Y,	0,	0	),
-			FPlane(	T3.X - T1.X,	T3.Y - T1.Y,	0,	0	),
-			FPlane(	T1.X,			T1.Y,			1,	0	),
-			FPlane(	0,				0,				0,	1	)
+			FPlane(T2.X - T1.X, T2.Y - T1.Y, 0, 0),
+			FPlane(T3.X - T1.X, T3.Y - T1.Y, 0, 0),
+			FPlane(T1.X, T1.Y, 1, 0),
+			FPlane(0, 0, 0, 1)
 			);
 
 		// Use InverseSlow to catch singular matrices.  Inverse can miss this sometimes.
 		const FMatrix TextureToLocal = ParameterToTexture.Inverse() * ParameterToLocal;
 
-		TriangleTangentX.Add(TextureToLocal.TransformVector(FVector(1,0,0)).GetSafeNormal());
-		TriangleTangentY.Add(TextureToLocal.TransformVector(FVector(0,1,0)).GetSafeNormal());
+		TriangleTangentX.Add(TextureToLocal.TransformVector(FVector(1, 0, 0)).GetSafeNormal());
+		TriangleTangentY.Add(TextureToLocal.TransformVector(FVector(0, 1, 0)).GetSafeNormal());
 		TriangleTangentZ.Add(Normal);
 
 		FVector::CreateOrthonormalBasis(
@@ -1790,12 +1793,12 @@ static void ComputeTriangleTangents(
 }
 
 /**
- * Create a table that maps the corner of each face to its overlapping corners.
- * @param OutOverlappingCorners - Maps a corner index to the indices of all overlapping corners.
- * @param RawMesh - The mesh for which to compute overlapping corners.
- */
+* Create a table that maps the corner of each face to its overlapping corners.
+* @param OutOverlappingCorners - Maps a corner index to the indices of all overlapping corners.
+* @param RawMesh - The mesh for which to compute overlapping corners.
+*/
 static void FindOverlappingCorners(
-	TMultiMap<int32,int32>& OutOverlappingCorners,
+	TMultiMap<int32, int32>& OutOverlappingCorners,
 	FRawMesh const& RawMesh,
 	float ComparisonThreshold
 	)
@@ -1805,9 +1808,9 @@ static void FindOverlappingCorners(
 	// Create a list of vertex Z/index pairs
 	TArray<FIndexAndZ> VertIndexAndZ;
 	VertIndexAndZ.Empty(NumWedges);
-	for(int32 WedgeIndex = 0; WedgeIndex < NumWedges; WedgeIndex++)
+	for (int32 WedgeIndex = 0; WedgeIndex < NumWedges; WedgeIndex++)
 	{
-		new(VertIndexAndZ) FIndexAndZ(WedgeIndex, GetPositionForWedge(RawMesh, WedgeIndex));
+		new(VertIndexAndZ)FIndexAndZ(WedgeIndex, GetPositionForWedge(RawMesh, WedgeIndex));
 	}
 
 	// Sort the vertices by z value
@@ -1825,10 +1828,10 @@ static void FindOverlappingCorners(
 			FVector PositionA = GetPositionForWedge(RawMesh, VertIndexAndZ[i].Index);
 			FVector PositionB = GetPositionForWedge(RawMesh, VertIndexAndZ[j].Index);
 
-			if (PointsEqual(PositionA, PositionB, ComparisonThreshold))					
+			if (PointsEqual(PositionA, PositionB, ComparisonThreshold))
 			{
-				OutOverlappingCorners.Add(VertIndexAndZ[i].Index,VertIndexAndZ[j].Index);
-				OutOverlappingCorners.Add(VertIndexAndZ[j].Index,VertIndexAndZ[i].Index);
+				OutOverlappingCorners.Add(VertIndexAndZ[i].Index, VertIndexAndZ[j].Index);
+				OutOverlappingCorners.Add(VertIndexAndZ[j].Index, VertIndexAndZ[i].Index);
 			}
 		}
 	}
@@ -1845,8 +1848,8 @@ namespace ETangentOptions
 };
 
 /**
-	* Smoothing group interpretation helper structure.
-	*/
+* Smoothing group interpretation helper structure.
+*/
 struct FFanFace
 {
 	int32 FaceIndex;
@@ -1858,7 +1861,7 @@ struct FFanFace
 
 static void ComputeTangents(
 	FRawMesh& RawMesh,
-	TMultiMap<int32,int32> const& OverlappingCorners,
+	TMultiMap<int32, int32> const& OverlappingCorners,
 	uint32 TangentOptions
 	)
 {
@@ -1922,15 +1925,15 @@ static void ComputeTangents(
 		}
 
 		// Don't process degenerate triangles.
-		if (PointsEqual(CornerPositions[0],CornerPositions[1], ComparisonThreshold)
-			|| PointsEqual(CornerPositions[0],CornerPositions[2], ComparisonThreshold)
-			|| PointsEqual(CornerPositions[1],CornerPositions[2], ComparisonThreshold))
+		if (PointsEqual(CornerPositions[0], CornerPositions[1], ComparisonThreshold)
+			|| PointsEqual(CornerPositions[0], CornerPositions[2], ComparisonThreshold)
+			|| PointsEqual(CornerPositions[1], CornerPositions[2], ComparisonThreshold))
 		{
 			continue;
 		}
 
 		// No need to process triangles if tangents already exist.
-		bool bCornerHasTangents[3] = {0};
+		bool bCornerHasTangents[3] = { 0 };
 		for (int32 CornerIndex = 0; CornerIndex < 3; CornerIndex++)
 		{
 			bCornerHasTangents[CornerIndex] = !RawMesh.WedgeTangentX[WedgeOffset + CornerIndex].IsZero()
@@ -1955,7 +1958,7 @@ static void ComputeTangents(
 		{
 			int32 ThisCornerIndex = WedgeOffset + CornerIndex;
 			DupVerts.Reset();
-			OverlappingCorners.MultiFind(ThisCornerIndex,DupVerts);
+			OverlappingCorners.MultiFind(ThisCornerIndex, DupVerts);
 			DupVerts.Add(ThisCornerIndex); // I am a "dup" of myself
 			for (int32 k = 0; k < DupVerts.Num(); k++)
 			{
@@ -1982,7 +1985,7 @@ static void ComputeTangents(
 				// Check for vertices in common.
 				if (FaceIndex == OtherFaceIndex)
 				{
-					CommonIndexCount = 3;		
+					CommonIndexCount = 3;
 					NewFanFace.LinkedVertexIndex = OurCornerIndex;
 				}
 				else
@@ -1991,10 +1994,10 @@ static void ComputeTangents(
 					for (int32 OtherCornerIndex = 0; OtherCornerIndex < 3; OtherCornerIndex++)
 					{
 						if (PointsEqual(
-								CornerPositions[OurCornerIndex],
-								GetPositionForWedge(RawMesh, OtherFaceIndex * 3 + OtherCornerIndex),
-								ComparisonThreshold
-								))
+							CornerPositions[OurCornerIndex],
+							GetPositionForWedge(RawMesh, OtherFaceIndex * 3 + OtherCornerIndex),
+							ComparisonThreshold
+							))
 						{
 							CommonIndexCount++;
 							NewFanFace.LinkedVertexIndex = OtherCornerIndex;
@@ -2004,7 +2007,7 @@ static void ComputeTangents(
 
 				// Add if connected by at least one point. Smoothing matches are considered later.
 				if (CommonIndexCount > 0)
-				{ 					
+				{
 					NewFanFace.FaceIndex = OtherFaceIndex;
 					NewFanFace.bFilled = (OtherFaceIndex == FaceIndex); // Starter face for smoothing floodfill.
 					NewFanFace.bBlendTangents = NewFanFace.bFilled;
@@ -2025,12 +2028,12 @@ static void ComputeTangents(
 			do
 			{
 				NewConnections = 0;
-				for (int32 OtherFaceIdx=0; OtherFaceIdx < RelevantFacesForCorner[CornerIndex].Num(); OtherFaceIdx++)
+				for (int32 OtherFaceIdx = 0; OtherFaceIdx < RelevantFacesForCorner[CornerIndex].Num(); OtherFaceIdx++)
 				{
 					FFanFace& OtherFace = RelevantFacesForCorner[CornerIndex][OtherFaceIdx];
 					// The vertex' own face is initially the only face with bFilled == true.
 					if (OtherFace.bFilled)
-					{				
+					{
 						for (int32 NextFaceIndex = 0; NextFaceIndex < RelevantFacesForCorner[CornerIndex].Num(); NextFaceIndex++)
 						{
 							FFanFace& NextFace = RelevantFacesForCorner[CornerIndex][NextFaceIndex];
@@ -2038,25 +2041,25 @@ static void ComputeTangents(
 							{
 								if ((NextFaceIndex != OtherFaceIdx)
 									&& (RawMesh.FaceSmoothingMasks[NextFace.FaceIndex] & RawMesh.FaceSmoothingMasks[OtherFace.FaceIndex]))
-								{				
+								{
 									int32 CommonVertices = 0;
 									int32 CommonTangentVertices = 0;
 									int32 CommonNormalVertices = 0;
 									for (int32 OtherCornerIndex = 0; OtherCornerIndex < 3; OtherCornerIndex++)
-									{											
+									{
 										for (int32 NextCornerIndex = 0; NextCornerIndex < 3; NextCornerIndex++)
 										{
 											int32 NextVertexIndex = RawMesh.WedgeIndices[NextFace.FaceIndex * 3 + NextCornerIndex];
 											int32 OtherVertexIndex = RawMesh.WedgeIndices[OtherFace.FaceIndex * 3 + OtherCornerIndex];
 											if (PointsEqual(
-													RawMesh.VertexPositions[NextVertexIndex],
-													RawMesh.VertexPositions[OtherVertexIndex],
-													ComparisonThreshold))
+												RawMesh.VertexPositions[NextVertexIndex],
+												RawMesh.VertexPositions[OtherVertexIndex],
+												ComparisonThreshold))
 											{
 												CommonVertices++;
 												if (UVsEqual(
-														RawMesh.WedgeTexCoords[0][NextFace.FaceIndex * 3 + NextCornerIndex],
-														RawMesh.WedgeTexCoords[0][OtherFace.FaceIndex * 3 + OtherCornerIndex]))
+													RawMesh.WedgeTexCoords[0][NextFace.FaceIndex * 3 + NextCornerIndex],
+													RawMesh.WedgeTexCoords[0][OtherFace.FaceIndex * 3 + OtherCornerIndex]))
 												{
 													CommonTangentVertices++;
 												}
@@ -2066,7 +2069,7 @@ static void ComputeTangents(
 													CommonNormalVertices++;
 												}
 											}
-										}										
+										}
 									}
 									// Flood fill faces with more than one common vertices which must be touching edges.
 									if (CommonVertices > 1)
@@ -2088,7 +2091,7 @@ static void ComputeTangents(
 												NextFace.bBlendTangents = true;
 											}
 										}
-									}								
+									}
 								}
 							}
 						}
@@ -2173,7 +2176,7 @@ static void ComputeTangents(
 }
 
 /*------------------------------------------------------------------------------
-	MikkTSpace for computing tangents.
+MikkTSpace for computing tangents.
 ------------------------------------------------------------------------------*/
 
 static int MikkGetNumFaces(const SMikkTSpaceContext* Context)
@@ -2258,13 +2261,13 @@ public:
 		TArray<FVector>				&VertexTangentsZ
 		)
 		:
-		wedges			(Wedges),
-		faces			(Faces),
-		points			(Points),
-		bComputeNormals	(bInComputeNormals),
-		TangentsX		(VertexTangentsX),
-		TangentsY		(VertexTangentsY),
-		TangentsZ		(VertexTangentsZ)
+		wedges(Wedges),
+		faces(Faces),
+		points(Points),
+		bComputeNormals(bInComputeNormals),
+		TangentsX(VertexTangentsX),
+		TangentsY(VertexTangentsY),
+		TangentsZ(VertexTangentsZ)
 	{
 	}
 };
@@ -2295,7 +2298,7 @@ static void MikkGetNormal_Skeletal(const SMikkTSpaceContext* Context, float Norm
 	MikkTSpace_Skeletal_Mesh *UserData = (MikkTSpace_Skeletal_Mesh*)(Context->m_pUserData);
 	// Get different normals depending on whether they've been calculated or not.
 	if (UserData->bComputeNormals) {
-		FVector &VertexNormal = UserData->TangentsZ[FaceIdx*3 + VertIdx];
+		FVector &VertexNormal = UserData->TangentsZ[FaceIdx * 3 + VertIdx];
 		Normal[0] = VertexNormal.X;
 		Normal[1] = VertexNormal.Y;
 		Normal[2] = VertexNormal.Z;
@@ -2312,7 +2315,7 @@ static void MikkGetNormal_Skeletal(const SMikkTSpaceContext* Context, float Norm
 static void MikkSetTSpaceBasic_Skeletal(const SMikkTSpaceContext* Context, const float Tangent[3], const float BitangentSign, const int FaceIdx, const int VertIdx)
 {
 	MikkTSpace_Skeletal_Mesh *UserData = (MikkTSpace_Skeletal_Mesh*)(Context->m_pUserData);
-	FVector &VertexTangent = UserData->TangentsX[FaceIdx*3 + VertIdx];
+	FVector &VertexTangent = UserData->TangentsX[FaceIdx * 3 + VertIdx];
 	VertexTangent.X = Tangent[0];
 	VertexTangent.Y = Tangent[1];
 	VertexTangent.Z = Tangent[2];
@@ -2326,7 +2329,7 @@ static void MikkSetTSpaceBasic_Skeletal(const SMikkTSpaceContext* Context, const
 	{
 		Bitangent = BitangentSign * FVector::CrossProduct(UserData->faces[FaceIdx].TangentZ[VertIdx], VertexTangent);
 	}
-	FVector &VertexBitangent = UserData->TangentsY[FaceIdx*3 + VertIdx];
+	FVector &VertexBitangent = UserData->TangentsY[FaceIdx * 3 + VertIdx];
 	// Switch the tangent space swizzle to X+Y-Z+ for legacy reasons.
 	VertexBitangent.X = -Bitangent[0];
 	VertexBitangent.Y = -Bitangent[1];
@@ -2343,7 +2346,7 @@ static void MikkGetTexCoord_Skeletal(const SMikkTSpaceContext* Context, float UV
 
 static void ComputeTangents_MikkTSpace(
 	FRawMesh& RawMesh,
-	TMultiMap<int32,int32> const& OverlappingCorners,
+	TMultiMap<int32, int32> const& OverlappingCorners,
 	uint32 TangentOptions
 	)
 {
@@ -2435,7 +2438,7 @@ static void ComputeTangents_MikkTSpace(
 			{
 				int32 ThisCornerIndex = WedgeOffset + CornerIndex;
 				DupVerts.Reset();
-				OverlappingCorners.MultiFind(ThisCornerIndex,DupVerts);
+				OverlappingCorners.MultiFind(ThisCornerIndex, DupVerts);
 				DupVerts.Add(ThisCornerIndex); // I am a "dup" of myself
 				for (int32 k = 0; k < DupVerts.Num(); k++)
 				{
@@ -2553,7 +2556,7 @@ static void ComputeTangents_MikkTSpace(
 							}
 						}
 					}
-				}
+				} 
 				while (NewConnections > 0);
 			}
 
@@ -2622,7 +2625,7 @@ static void ComputeTangents_MikkTSpace(
 		MikkTInterface.m_getTexCoord = MikkGetTexCoord;
 		MikkTInterface.m_setTSpaceBasic = MikkSetTSpaceBasic;
 		MikkTInterface.m_setTSpace = nullptr;
-		
+
 		SMikkTSpaceContext MikkTContext;
 		MikkTContext.m_pInterface = &MikkTInterface;
 		MikkTContext.m_pUserData = (void*)(&RawMesh);
@@ -2655,25 +2658,25 @@ static void ComputeStreamingTextureFactors(
 		const FVector& Pos0 = Mesh.GetWedgePosition(Wedge0) * BuildScale;
 		const FVector& Pos1 = Mesh.GetWedgePosition(Wedge1) * BuildScale;
 		const FVector& Pos2 = Mesh.GetWedgePosition(Wedge2) * BuildScale;
-		float	L1	= (Pos0 - Pos1).Size(),
-				L2	= (Pos0 - Pos2).Size();
+		float	L1 = (Pos0 - Pos1).Size(),
+			L2 = (Pos0 - Pos2).Size();
 
-		for(int32 UVIndex = 0;UVIndex < NumTexCoords;UVIndex++)
+		for (int32 UVIndex = 0; UVIndex < NumTexCoords; UVIndex++)
 		{
 			FVector2D UV0 = Mesh.WedgeTexCoords[UVIndex][Wedge0];
 			FVector2D UV1 = Mesh.WedgeTexCoords[UVIndex][Wedge1];
 			FVector2D UV2 = Mesh.WedgeTexCoords[UVIndex][Wedge2];
 
-			float	T1	= (UV0 - UV1).Size(),
-					T2	= (UV0 - UV2).Size();
+			float	T1 = (UV0 - UV1).Size(),
+				T2 = (UV0 - UV2).Size();
 
-			if( FMath::Abs(T1 * T2) > FMath::Square(SMALL_NUMBER) )
+			if (FMath::Abs(T1 * T2) > FMath::Square(SMALL_NUMBER))
 			{
-				const float TexelRatio = FMath::Max( L1 / T1, L2 / T2 );
-				TexelRatios[UVIndex].Add( TexelRatio );
+				const float TexelRatio = FMath::Max(L1 / T1, L2 / T2);
+				TexelRatios[UVIndex].Add(TexelRatio);
 
 				// Update max texel ratio
-				if( TexelRatio > MaxStreamingTextureFactor )
+				if (TexelRatio > MaxStreamingTextureFactor)
 				{
 					MaxStreamingTextureFactor = TexelRatio;
 				}
@@ -2681,15 +2684,15 @@ static void ComputeStreamingTextureFactors(
 		}
 	}
 
-	for(int32 UVIndex = 0;UVIndex < MAX_STATIC_TEXCOORDS;UVIndex++)
+	for (int32 UVIndex = 0; UVIndex < MAX_STATIC_TEXCOORDS; UVIndex++)
 	{
 		OutStreamingTextureFactors[UVIndex] = 0.0f;
-		if( TexelRatios[UVIndex].Num() )
+		if (TexelRatios[UVIndex].Num())
 		{
 			// Disregard upper 75% of texel ratios.
 			// This is to ignore backfacing surfaces or other non-visible surfaces that tend to map a small number of texels to a large surface.
-			TexelRatios[UVIndex].Sort( TGreater<float>() );
-			float TexelRatio = TexelRatios[UVIndex][ FMath::TruncToInt(TexelRatios[UVIndex].Num() * 0.75f) ];
+			TexelRatios[UVIndex].Sort(TGreater<float>());
+			float TexelRatio = TexelRatios[UVIndex][FMath::TruncToInt(TexelRatios[UVIndex].Num() * 0.75f)];
 			OutStreamingTextureFactors[UVIndex] = TexelRatio;
 		}
 	}
@@ -2713,16 +2716,16 @@ static void BuildDepthOnlyIndexBuffer(
 	// Create a mapping of index -> first overlapping index to accelerate the construction of the shadow index buffer.
 	TArray<FIndexAndZ> VertIndexAndZ;
 	VertIndexAndZ.Empty(NumVertices);
-	for(int32 VertIndex = 0; VertIndex < NumVertices; VertIndex++)
+	for (int32 VertIndex = 0; VertIndex < NumVertices; VertIndex++)
 	{
-		new(VertIndexAndZ) FIndexAndZ(VertIndex, InVertices[VertIndex].Position);
+		new(VertIndexAndZ)FIndexAndZ(VertIndex, InVertices[VertIndex].Position);
 	}
 	VertIndexAndZ.Sort(FCompareIndexAndZ());
 
 	// Setup the index map. 0xFFFFFFFF == not set.
 	TArray<uint32> IndexMap;
 	IndexMap.AddUninitialized(NumVertices);
-	FMemory::Memset(IndexMap.GetData(),0xFF,NumVertices * sizeof(uint32));
+	FMemory::Memset(IndexMap.GetData(), 0xFF, NumVertices * sizeof(uint32));
 
 	// Search for duplicates, quickly!
 	for (int32 i = 0; i < VertIndexAndZ.Num(); i++)
@@ -2738,7 +2741,7 @@ static void BuildDepthOnlyIndexBuffer(
 				break; // can't be any more dups
 
 			uint32 OtherIndex = VertIndexAndZ[j].Index;
-			if (PointsEqual(InVertices[SrcIndex].Position,InVertices[OtherIndex].Position,/*bUseEpsilonCompare=*/ true))					
+			if (PointsEqual(InVertices[SrcIndex].Position, InVertices[OtherIndex].Position,/*bUseEpsilonCompare=*/ true))
 			{
 				IndexMap[SrcIndex] = FMath::Min(IndexMap[SrcIndex], OtherIndex);
 				IndexMap[OtherIndex] = FMath::Min(IndexMap[OtherIndex], SrcIndex);
@@ -2768,7 +2771,7 @@ static float GetComparisonThreshold(FMeshBuildSettings const& BuildSettings)
 }
 
 /*------------------------------------------------------------------------------
-	Static mesh building.
+Static mesh building.
 ------------------------------------------------------------------------------*/
 
 static FStaticMeshBuildVertex BuildStaticMeshVertex(FRawMesh const& RawMesh, int32 WedgeIndex, FVector BuildScale)
@@ -2776,7 +2779,7 @@ static FStaticMeshBuildVertex BuildStaticMeshVertex(FRawMesh const& RawMesh, int
 	FStaticMeshBuildVertex Vertex;
 	Vertex.Position = GetPositionForWedge(RawMesh, WedgeIndex) * BuildScale;
 
-	const FMatrix ScaleMatrix = FScaleMatrix( BuildScale ).Inverse().GetTransposed();	
+	const FMatrix ScaleMatrix = FScaleMatrix(BuildScale).Inverse().GetTransposed();
 	Vertex.TangentX = ScaleMatrix.TransformVector(RawMesh.WedgeTangentX[WedgeIndex]).GetSafeNormal();
 	Vertex.TangentY = ScaleMatrix.TransformVector(RawMesh.WedgeTangentY[WedgeIndex]).GetSafeNormal();
 	Vertex.TangentZ = ScaleMatrix.TransformVector(RawMesh.WedgeTangentZ[WedgeIndex]).GetSafeNormal();
@@ -2799,7 +2802,7 @@ static FStaticMeshBuildVertex BuildStaticMeshVertex(FRawMesh const& RawMesh, int
 		}
 		else
 		{
-			Vertex.UVs[i] = FVector2D(0.0f,0.0f);
+			Vertex.UVs[i] = FVector2D(0.0f, 0.0f);
 		}
 	}
 	return Vertex;
@@ -2837,17 +2840,17 @@ static void BuildStaticMeshVertexAndIndexBuffers(
 	TArray<TArray<uint32> >& OutPerSectionIndices,
 	TArray<int32>& OutWedgeMap,
 	const FRawMesh& RawMesh,
-	const TMultiMap<int32,int32>& OverlappingCorners,
+	const TMultiMap<int32, int32>& OverlappingCorners,
 	float ComparisonThreshold,
 	FVector BuildScale
 	)
 {
-	TMap<int32,int32> FinalVerts;
+	TMap<int32, int32> FinalVerts;
 	TArray<int32> DupVerts;
 	int32 NumFaces = RawMesh.WedgeIndices.Num() / 3;
 
 	// Process each face, build vertex buffer and per-section index buffers.
-	for(int32 FaceIndex = 0; FaceIndex < NumFaces; FaceIndex++)
+	for (int32 FaceIndex = 0; FaceIndex < NumFaces; FaceIndex++)
 	{
 		int32 VertexIndices[3];
 		FVector CornerPositions[3];
@@ -2858,9 +2861,9 @@ static void BuildStaticMeshVertexAndIndexBuffers(
 		}
 
 		// Don't process degenerate triangles.
-		if (PointsEqual(CornerPositions[0],CornerPositions[1], ComparisonThreshold)
-			|| PointsEqual(CornerPositions[0],CornerPositions[2], ComparisonThreshold)
-			|| PointsEqual(CornerPositions[1],CornerPositions[2], ComparisonThreshold))
+		if (PointsEqual(CornerPositions[0], CornerPositions[1], ComparisonThreshold)
+			|| PointsEqual(CornerPositions[0], CornerPositions[2], ComparisonThreshold)
+			|| PointsEqual(CornerPositions[1], CornerPositions[2], ComparisonThreshold))
 		{
 			for (int32 CornerIndex = 0; CornerIndex < 3; CornerIndex++)
 			{
@@ -2875,7 +2878,7 @@ static void BuildStaticMeshVertexAndIndexBuffers(
 			FStaticMeshBuildVertex ThisVertex = BuildStaticMeshVertex(RawMesh, WedgeIndex, BuildScale);
 
 			DupVerts.Reset();
-			OverlappingCorners.MultiFind(WedgeIndex,DupVerts);
+			OverlappingCorners.MultiFind(WedgeIndex, DupVerts);
 			DupVerts.Sort();
 
 			int32 Index = INDEX_NONE;
@@ -2898,7 +2901,7 @@ static void BuildStaticMeshVertexAndIndexBuffers(
 			if (Index == INDEX_NONE)
 			{
 				Index = OutVertices.Add(ThisVertex);
-				FinalVerts.Add(WedgeIndex,Index);
+				FinalVerts.Add(WedgeIndex, Index);
 			}
 			VertexIndices[CornerIndex] = Index;
 		}
@@ -2916,7 +2919,7 @@ static void BuildStaticMeshVertexAndIndexBuffers(
 		}
 
 		// Put the indices in the material index buffer.
-		int32 SectionIndex = FMath::Clamp(RawMesh.FaceMaterialIndices[FaceIndex], 0, OutPerSectionIndices.Num()-1);
+		int32 SectionIndex = FMath::Clamp(RawMesh.FaceMaterialIndices[FaceIndex], 0, OutPerSectionIndices.Num() - 1);
 		TArray<uint32>& SectionIndices = OutPerSectionIndices[SectionIndex];
 		for (int32 CornerIndex = 0; CornerIndex < 3; CornerIndex++)
 		{
@@ -3088,9 +3091,9 @@ public:
 			{
 				// If a raw mesh is not explicitly provided, use the raw mesh of the
 				// next highest LOD.
-				RawMesh = LODMeshes[LODIndex-1];
-				OverlappingCorners = LODOverlappingCorners[LODIndex-1];
-				LODBuildSettings[LODIndex] = LODBuildSettings[LODIndex-1];
+				RawMesh = LODMeshes[LODIndex - 1];
+				OverlappingCorners = LODOverlappingCorners[LODIndex - 1];
+				LODBuildSettings[LODIndex] = LODBuildSettings[LODIndex - 1];
 				HasRawMesh[LODIndex] = false;
 			}
 		}
@@ -3314,7 +3317,7 @@ public:
 
 				FStaticMeshEdgeBuilder(CombinedIndices, Vertices, Edges).FindEdges();
 				WireframeIndices.Empty(2 * Edges.Num());
-				for (int32 EdgeIndex = 0;EdgeIndex < Edges.Num();EdgeIndex++)
+				for (int32 EdgeIndex = 0; EdgeIndex < Edges.Num(); EdgeIndex++)
 				{
 					FMeshEdge&	Edge = Edges[EdgeIndex];
 					WireframeIndices.Add(Edge.Vertices[0]);
@@ -3478,10 +3481,10 @@ protected:
 		const IMeshUtilities::MeshBuildOptions& InBuildOptions,
 		TArray<FText>* InWarningMessages,
 		TArray<FName>* InWarningNames)
-	: BuildOptions(InBuildOptions)
-	, OutWarningMessages(InWarningMessages)
-	, OutWarningNames(InWarningNames)
-	, bTooManyVerts(false)
+		: BuildOptions(InBuildOptions)
+		, OutWarningMessages(InWarningMessages)
+		, OutWarningNames(InWarningNames)
+		, bTooManyVerts(false)
 	{
 	}
 };
@@ -3500,15 +3503,15 @@ public:
 		const IMeshUtilities::MeshBuildOptions& InBuildOptions,
 		TArray<FText>* InWarningMessages,
 		TArray<FName>* InWarningNames)
-	: IMeshBuildData(InBuildOptions, InWarningMessages, InWarningNames)
-	, MikkTUserData(InWedges, InFaces, InPoints, InBuildOptions.bComputeNormals, TangentX, TangentY, TangentZ)
-	, LODModel(InLODModel)
-	, RefSkeleton(InRefSkeleton)
-	, Influences(InInfluences)
-	, Wedges(InWedges)
-	, Faces(InFaces)
-	, Points(InPoints)
-	, PointToOriginalMap(InPointToOriginalMap)	
+		: IMeshBuildData(InBuildOptions, InWarningMessages, InWarningNames)
+		, MikkTUserData(InWedges, InFaces, InPoints, InBuildOptions.bComputeNormals, TangentX, TangentY, TangentZ)
+		, LODModel(InLODModel)
+		, RefSkeleton(InRefSkeleton)
+		, Influences(InInfluences)
+		, Wedges(InWedges)
+		, Faces(InFaces)
+		, Points(InPoints)
+		, PointToOriginalMap(InPointToOriginalMap)
 	{
 		MikkTInterface.m_getNormal = MikkGetNormal_Skeletal;
 		MikkTInterface.m_getNumFaces = MikkGetNumFaces_Skeletal;
@@ -3530,7 +3533,7 @@ public:
 	}
 
 	virtual uint32 GetVertexIndex(uint32 FaceIndex, uint32 TriIndex) override
-	{	
+	{
 		return Wedges[Faces[FaceIndex].iWedge[TriIndex]].iVertex;
 	}
 
@@ -3611,13 +3614,13 @@ class FSkeletalMeshUtilityBuilder
 {
 public:
 	FSkeletalMeshUtilityBuilder()
-	: Stage(EStage::Uninit)
+		: Stage(EStage::Uninit)
 	{
 	}
 
 public:
 	void Skeletal_FindOverlappingCorners(
-		TMultiMap<int32,int32>& OutOverlappingCorners,
+		TMultiMap<int32, int32>& OutOverlappingCorners,
 		IMeshBuildData* BuildData,
 		float ComparisonThreshold
 		)
@@ -3629,13 +3632,13 @@ public:
 		// Create a list of vertex Z/index pairs
 		TArray<FIndexAndZ> VertIndexAndZ;
 		VertIndexAndZ.Empty(NumWedges);
-		for(int32 FaceIndex = 0; FaceIndex < NumFaces; FaceIndex++)
+		for (int32 FaceIndex = 0; FaceIndex < NumFaces; FaceIndex++)
 		{
 			for (int32 TriIndex = 0; TriIndex < 3; ++TriIndex)
 			{
 				uint32 Index = BuildData->GetWedgeIndex(FaceIndex, TriIndex);
-				new(VertIndexAndZ) FIndexAndZ(Index, BuildData->GetVertexPosition(Index));
-			}	
+				new(VertIndexAndZ)FIndexAndZ(Index, BuildData->GetVertexPosition(Index));
+			}
 		}
 
 		// Sort the vertices by z value
@@ -3653,10 +3656,10 @@ public:
 				FVector PositionA = BuildData->GetVertexPosition(VertIndexAndZ[i].Index);
 				FVector PositionB = BuildData->GetVertexPosition(VertIndexAndZ[j].Index);
 
-				if (PointsEqual(PositionA, PositionB, ComparisonThreshold))					
+				if (PointsEqual(PositionA, PositionB, ComparisonThreshold))
 				{
-					OutOverlappingCorners.Add(VertIndexAndZ[i].Index,VertIndexAndZ[j].Index);
-					OutOverlappingCorners.Add(VertIndexAndZ[j].Index,VertIndexAndZ[i].Index);
+					OutOverlappingCorners.Add(VertIndexAndZ[i].Index, VertIndexAndZ[j].Index);
+					OutOverlappingCorners.Add(VertIndexAndZ[j].Index, VertIndexAndZ[i].Index);
 				}
 			}
 		}
@@ -3685,29 +3688,29 @@ public:
 				P[i] = BuildData->GetVertexPosition(TriangleIndex, i);
 			}
 
-			const FVector Normal = ((P[1] - P[2])^(P[0] - P[2])).GetSafeNormal(ComparisonThreshold);
+			const FVector Normal = ((P[1] - P[2]) ^ (P[0] - P[2])).GetSafeNormal(ComparisonThreshold);
 			FMatrix	ParameterToLocal(
 				FPlane(P[1].X - P[0].X, P[1].Y - P[0].Y, P[1].Z - P[0].Z, 0),
 				FPlane(P[2].X - P[0].X, P[2].Y - P[0].Y, P[2].Z - P[0].Z, 0),
-				FPlane(P[0].X,          P[0].Y,          P[0].Z,          0),
-				FPlane(0,               0,               0,				  1)
+				FPlane(P[0].X, P[0].Y, P[0].Z, 0),
+				FPlane(0, 0, 0, 1)
 				);
 
 			FVector2D T1 = BuildData->GetVertexUV(TriangleIndex, 0, UVIndex);
 			FVector2D T2 = BuildData->GetVertexUV(TriangleIndex, 1, UVIndex);
 			FVector2D T3 = BuildData->GetVertexUV(TriangleIndex, 2, UVIndex);
 			FMatrix ParameterToTexture(
-				FPlane(	T2.X - T1.X,	T2.Y - T1.Y,	0,	0	),
-				FPlane(	T3.X - T1.X,	T3.Y - T1.Y,	0,	0	),
-				FPlane(	T1.X,			T1.Y,			1,	0	),
-				FPlane(	0,				0,				0,	1	)
+				FPlane(T2.X - T1.X, T2.Y - T1.Y, 0, 0),
+				FPlane(T3.X - T1.X, T3.Y - T1.Y, 0, 0),
+				FPlane(T1.X, T1.Y, 1, 0),
+				FPlane(0, 0, 0, 1)
 				);
 
 			// Use InverseSlow to catch singular matrices.  Inverse can miss this sometimes.
 			const FMatrix TextureToLocal = ParameterToTexture.Inverse() * ParameterToLocal;
 
-			TriangleTangentX.Add(TextureToLocal.TransformVector(FVector(1,0,0)).GetSafeNormal());
-			TriangleTangentY.Add(TextureToLocal.TransformVector(FVector(0,1,0)).GetSafeNormal());
+			TriangleTangentX.Add(TextureToLocal.TransformVector(FVector(1, 0, 0)).GetSafeNormal());
+			TriangleTangentY.Add(TextureToLocal.TransformVector(FVector(0, 1, 0)).GetSafeNormal());
 			TriangleTangentZ.Add(Normal);
 
 			FVector::CreateOrthonormalBasis(
@@ -3720,7 +3723,7 @@ public:
 
 	void Skeletal_ComputeTangents(
 		IMeshBuildData* BuildData,
-		TMultiMap<int32,int32> const& OverlappingCorners
+		TMultiMap<int32, int32> const& OverlappingCorners
 		)
 	{
 		bool bBlendOverlappingNormals = true;
@@ -3788,15 +3791,15 @@ public:
 			}
 
 			// Don't process degenerate triangles.
-			if (PointsEqual(CornerPositions[0],CornerPositions[1], ComparisonThreshold)
-				|| PointsEqual(CornerPositions[0],CornerPositions[2], ComparisonThreshold)
-				|| PointsEqual(CornerPositions[1],CornerPositions[2], ComparisonThreshold))
+			if (PointsEqual(CornerPositions[0], CornerPositions[1], ComparisonThreshold)
+				|| PointsEqual(CornerPositions[0], CornerPositions[2], ComparisonThreshold)
+				|| PointsEqual(CornerPositions[1], CornerPositions[2], ComparisonThreshold))
 			{
 				continue;
 			}
 
 			// No need to process triangles if tangents already exist.
-			bool bCornerHasTangents[3] = {0};
+			bool bCornerHasTangents[3] = { 0 };
 			for (int32 CornerIndex = 0; CornerIndex < 3; CornerIndex++)
 			{
 				bCornerHasTangents[CornerIndex] = !WedgeTangentX[WedgeOffset + CornerIndex].IsZero()
@@ -3821,7 +3824,7 @@ public:
 			{
 				int32 ThisCornerIndex = WedgeOffset + CornerIndex;
 				DupVerts.Reset();
-				OverlappingCorners.MultiFind(ThisCornerIndex,DupVerts);
+				OverlappingCorners.MultiFind(ThisCornerIndex, DupVerts);
 				DupVerts.Add(ThisCornerIndex); // I am a "dup" of myself
 				for (int32 k = 0; k < DupVerts.Num(); k++)
 				{
@@ -3848,7 +3851,7 @@ public:
 					// Check for vertices in common.
 					if (FaceIndex == OtherFaceIndex)
 					{
-						CommonIndexCount = 3;		
+						CommonIndexCount = 3;
 						NewFanFace.LinkedVertexIndex = OurCornerIndex;
 					}
 					else
@@ -3857,10 +3860,10 @@ public:
 						for (int32 OtherCornerIndex = 0; OtherCornerIndex < 3; OtherCornerIndex++)
 						{
 							if (PointsEqual(
-									CornerPositions[OurCornerIndex],
-									BuildData->GetVertexPosition(OtherFaceIndex, OtherCornerIndex),
-									ComparisonThreshold
-									))
+								CornerPositions[OurCornerIndex],
+								BuildData->GetVertexPosition(OtherFaceIndex, OtherCornerIndex),
+								ComparisonThreshold
+								))
 							{
 								CommonIndexCount++;
 								NewFanFace.LinkedVertexIndex = OtherCornerIndex;
@@ -3870,7 +3873,7 @@ public:
 
 					// Add if connected by at least one point. Smoothing matches are considered later.
 					if (CommonIndexCount > 0)
-					{ 					
+					{
 						NewFanFace.FaceIndex = OtherFaceIndex;
 						NewFanFace.bFilled = (OtherFaceIndex == FaceIndex); // Starter face for smoothing floodfill.
 						NewFanFace.bBlendTangents = NewFanFace.bFilled;
@@ -3891,12 +3894,12 @@ public:
 				do
 				{
 					NewConnections = 0;
-					for (int32 OtherFaceIdx=0; OtherFaceIdx < RelevantFacesForCorner[CornerIndex].Num(); OtherFaceIdx++)
+					for (int32 OtherFaceIdx = 0; OtherFaceIdx < RelevantFacesForCorner[CornerIndex].Num(); OtherFaceIdx++)
 					{
 						FFanFace& OtherFace = RelevantFacesForCorner[CornerIndex][OtherFaceIdx];
 						// The vertex' own face is initially the only face with bFilled == true.
 						if (OtherFace.bFilled)
-						{				
+						{
 							for (int32 NextFaceIndex = 0; NextFaceIndex < RelevantFacesForCorner[CornerIndex].Num(); NextFaceIndex++)
 							{
 								FFanFace& NextFace = RelevantFacesForCorner[CornerIndex][NextFaceIndex];
@@ -3904,12 +3907,12 @@ public:
 								{
 									if (NextFaceIndex != OtherFaceIdx)
 										//&& (RawMesh.FaceSmoothingMasks[NextFace.FaceIndex] & RawMesh.FaceSmoothingMasks[OtherFace.FaceIndex]))
-									{				
+									{
 										int32 CommonVertices = 0;
 										int32 CommonTangentVertices = 0;
 										int32 CommonNormalVertices = 0;
 										for (int32 OtherCornerIndex = 0; OtherCornerIndex < 3; OtherCornerIndex++)
-										{											
+										{
 											for (int32 NextCornerIndex = 0; NextCornerIndex < 3; NextCornerIndex++)
 											{
 												int32 NextVertexIndex = BuildData->GetVertexIndex(NextFace.FaceIndex, NextCornerIndex);
@@ -3921,10 +3924,10 @@ public:
 												{
 													CommonVertices++;
 
-													
+
 													if (UVsEqual(
-															BuildData->GetVertexUV(NextFace.FaceIndex, NextCornerIndex, 0),
-															BuildData->GetVertexUV(OtherFace.FaceIndex, OtherCornerIndex, 0)))
+														BuildData->GetVertexUV(NextFace.FaceIndex, NextCornerIndex, 0),
+														BuildData->GetVertexUV(OtherFace.FaceIndex, OtherCornerIndex, 0)))
 													{
 														CommonTangentVertices++;
 													}
@@ -3934,7 +3937,7 @@ public:
 														CommonNormalVertices++;
 													}
 												}
-											}										
+											}
 										}
 										// Flood fill faces with more than one common vertices which must be touching edges.
 										if (CommonVertices > 1)
@@ -3956,7 +3959,7 @@ public:
 													NextFace.bBlendTangents = true;
 												}
 											}
-										}								
+										}
 									}
 								}
 							}
@@ -4042,7 +4045,7 @@ public:
 
 	void Skeletal_ComputeTangents_MikkTSpace(
 		IMeshBuildData* BuildData,
-		TMultiMap<int32,int32> const& OverlappingCorners
+		TMultiMap<int32, int32> const& OverlappingCorners
 		)
 	{
 		bool bBlendOverlappingNormals = true;
@@ -4079,8 +4082,8 @@ public:
 		bool bWedgeTSpace = false;
 		for (int32 WedgeIdx = 0; WedgeIdx < WedgeTangentZ.Num(); ++WedgeIdx)
 		{
-			for(int32 VertexIndex = 0;VertexIndex < 3;VertexIndex++)
-			bWedgeNormals = bWedgeNormals && (!WedgeTangentZ[WedgeIdx].IsNearlyZero());
+			for (int32 VertexIndex = 0; VertexIndex < 3; VertexIndex++)
+				bWedgeNormals = bWedgeNormals && (!WedgeTangentZ[WedgeIdx].IsNearlyZero());
 		}
 
 		if (WedgeTangentX.Num() > 0 && WedgeTangentY.Num() > 0)
@@ -4139,7 +4142,7 @@ public:
 				{
 					int32 ThisCornerIndex = WedgeOffset + CornerIndex;
 					DupVerts.Reset();
-					OverlappingCorners.MultiFind(ThisCornerIndex,DupVerts);
+					OverlappingCorners.MultiFind(ThisCornerIndex, DupVerts);
 					DupVerts.Add(ThisCornerIndex); // I am a "dup" of myself
 					for (int32 k = 0; k < DupVerts.Num(); k++)
 					{
@@ -4257,7 +4260,7 @@ public:
 								}
 							}
 						}
-					}
+					} 
 					while (NewConnections > 0);
 				}
 
@@ -4393,27 +4396,27 @@ public:
 		SkeletalMeshBuildData& BuildData = *(SkeletalMeshBuildData*)InBuildData;
 
 		BeginSlowTask();
-	
+
 		// Find wedge influences.
 		TArray<int32>	WedgeInfluenceIndices;
-		TMap<uint32,uint32> VertexIndexToInfluenceIndexMap;
+		TMap<uint32, uint32> VertexIndexToInfluenceIndexMap;
 
-		for(uint32 LookIdx = 0;LookIdx < (uint32)BuildData.Influences.Num();LookIdx++)
+		for (uint32 LookIdx = 0; LookIdx < (uint32)BuildData.Influences.Num(); LookIdx++)
 		{
 			// Order matters do not allow the map to overwrite an existing value.
-			if( !VertexIndexToInfluenceIndexMap.Find( BuildData.Influences[LookIdx].VertIndex) )
+			if (!VertexIndexToInfluenceIndexMap.Find(BuildData.Influences[LookIdx].VertIndex))
 			{
-				VertexIndexToInfluenceIndexMap.Add( BuildData.Influences[LookIdx].VertIndex, LookIdx );
+				VertexIndexToInfluenceIndexMap.Add(BuildData.Influences[LookIdx].VertIndex, LookIdx);
 			}
 		}
 
-		for(int32 WedgeIndex = 0;WedgeIndex < BuildData.Wedges.Num();WedgeIndex++)
+		for (int32 WedgeIndex = 0; WedgeIndex < BuildData.Wedges.Num(); WedgeIndex++)
 		{
-			uint32* InfluenceIndex = VertexIndexToInfluenceIndexMap.Find( BuildData.Wedges[WedgeIndex].iVertex );
+			uint32* InfluenceIndex = VertexIndexToInfluenceIndexMap.Find(BuildData.Wedges[WedgeIndex].iVertex);
 
-			if ( InfluenceIndex )
+			if (InfluenceIndex)
 			{
-				WedgeInfluenceIndices.Add( *InfluenceIndex );
+				WedgeInfluenceIndices.Add(*InfluenceIndex);
 			}
 			else
 			{
@@ -4440,7 +4443,7 @@ public:
 		VertIndexAndZ.Empty(BuildData.Points.Num());
 		RawVertices.Reserve(BuildData.Points.Num());
 
-		for(int32 FaceIndex = 0; FaceIndex < BuildData.Faces.Num(); FaceIndex++)
+		for (int32 FaceIndex = 0; FaceIndex < BuildData.Faces.Num(); FaceIndex++)
 		{
 			// Only update the status progress bar if we are in the game thread and every thousand faces. 
 			// Updating status is extremely slow
@@ -4450,8 +4453,8 @@ public:
 			}
 
 			const FMeshFace& Face = BuildData.Faces[FaceIndex];
-						
-			for(int32 VertexIndex = 0; VertexIndex < 3; VertexIndex++)
+
+			for (int32 VertexIndex = 0; VertexIndex < 3; VertexIndex++)
 			{
 				FSoftSkinBuildVertex Vertex;
 				const uint32 WedgeIndex = BuildData.GetWedgeIndex(FaceIndex, VertexIndex);
@@ -4459,33 +4462,33 @@ public:
 
 				Vertex.Position = BuildData.GetVertexPosition(FaceIndex, VertexIndex);
 
-				FVector TangentX,TangentY,TangentZ;
+				FVector TangentX, TangentY, TangentZ;
 				TangentX = BuildData.TangentX[WedgeIndex].GetSafeNormal();
 				TangentY = BuildData.TangentY[WedgeIndex].GetSafeNormal();
 				TangentZ = BuildData.TangentZ[WedgeIndex].GetSafeNormal();
 
 				/*if (BuildData.BuildOptions.bComputeNormals || BuildData.BuildOptions.bComputeTangents)
 				{
-					TangentX = BuildData.TangentX[VertexIndex].GetSafeNormal();
-					TangentY = BuildData.TangentY[VertexIndex].GetSafeNormal();
+				TangentX = BuildData.TangentX[VertexIndex].GetSafeNormal();
+				TangentY = BuildData.TangentY[VertexIndex].GetSafeNormal();
 
-					if( BuildData.BuildOptions.bComputeNormals )
-					{
-						TangentZ = BuildData.TangentZ[VertexIndex].GetSafeNormal();
-					}
-					else
-					{
-						//TangentZ = Face.TangentZ[VertexIndex];
-					}
+				if( BuildData.BuildOptions.bComputeNormals )
+				{
+				TangentZ = BuildData.TangentZ[VertexIndex].GetSafeNormal();
+				}
+				else
+				{
+				//TangentZ = Face.TangentZ[VertexIndex];
+				}
 
-					TangentY -= TangentX * (TangentX | TangentY);
-					TangentY.Normalize();
+				TangentY -= TangentX * (TangentX | TangentY);
+				TangentY.Normalize();
 
-					TangentX -= TangentZ * (TangentZ | TangentX);
-					TangentY -= TangentZ * (TangentZ | TangentY);
+				TangentX -= TangentZ * (TangentZ | TangentX);
+				TangentY -= TangentZ * (TangentZ | TangentY);
 
-					TangentX.Normalize();
-					TangentY.Normalize();
+				TangentX.Normalize();
+				TangentY.Normalize();
 				}
 				else*/
 				{
@@ -4503,7 +4506,7 @@ public:
 				Vertex.TangentY = TangentY;
 				Vertex.TangentZ = TangentZ;
 
-				FMemory::Memcpy(Vertex.UVs, Wedge.UVs, sizeof(FVector2D)*MAX_TEXCOORDS);	
+				FMemory::Memcpy(Vertex.UVs, Wedge.UVs, sizeof(FVector2D)*MAX_TEXCOORDS);
 				Vertex.Color = Wedge.Color;
 
 				{
@@ -4512,31 +4515,31 @@ public:
 					int32 LookIdx = InfIdx;
 
 					uint32 InfluenceCount = 0;
-					while( BuildData.Influences.IsValidIndex(LookIdx) && (BuildData.Influences[LookIdx].VertIndex == Wedge.iVertex) )
-					{			
+					while (BuildData.Influences.IsValidIndex(LookIdx) && (BuildData.Influences[LookIdx].VertIndex == Wedge.iVertex))
+					{
 						InfluenceCount++;
 						LookIdx++;
 					}
-					InfluenceCount = FMath::Min<uint32>(InfluenceCount,MAX_TOTAL_INFLUENCES);
+					InfluenceCount = FMath::Min<uint32>(InfluenceCount, MAX_TOTAL_INFLUENCES);
 
 					// Setup the vertex influences.
 					Vertex.InfluenceBones[0] = 0;
 					Vertex.InfluenceWeights[0] = 255;
-					for(uint32 i = 1;i < MAX_TOTAL_INFLUENCES;i++)
+					for (uint32 i = 1; i < MAX_TOTAL_INFLUENCES; i++)
 					{
 						Vertex.InfluenceBones[i] = 0;
 						Vertex.InfluenceWeights[i] = 0;
 					}
 
 					uint32	TotalInfluenceWeight = 0;
-					for(uint32 i = 0;i < InfluenceCount;i++)
+					for (uint32 i = 0; i < InfluenceCount; i++)
 					{
-						FBoneIndexType BoneIndex = (FBoneIndexType)BuildData.Influences[InfIdx+i].BoneIndex;
-						if( BoneIndex >= BuildData.RefSkeleton.GetNum() )
+						FBoneIndexType BoneIndex = (FBoneIndexType)BuildData.Influences[InfIdx + i].BoneIndex;
+						if (BoneIndex >= BuildData.RefSkeleton.GetNum())
 							continue;
 
 						Vertex.InfluenceBones[i] = BoneIndex;
-						Vertex.InfluenceWeights[i] = (uint8)(BuildData.Influences[InfIdx+i].Weight * 255.0f);
+						Vertex.InfluenceWeights[i] = (uint8)(BuildData.Influences[InfIdx + i].Weight * 255.0f);
 						TotalInfluenceWeight += Vertex.InfluenceWeights[i];
 					}
 					Vertex.InfluenceWeights[0] += 255 - TotalInfluenceWeight;
@@ -4544,15 +4547,15 @@ public:
 
 				// Add the vertex as well as its original index in the points array
 				Vertex.PointWedgeIdx = Wedge.iVertex;
-			
-				int32 RawIndex = RawVertices.Add( Vertex );
+
+				int32 RawIndex = RawVertices.Add(Vertex);
 
 				// Add an efficient way to find dupes of this vertex later for fast combining of vertices
 				FSkeletalMeshVertIndexAndZ IAndZ;
 				IAndZ.Index = RawIndex;
 				IAndZ.Z = Vertex.Position.Z;
 
-				VertIndexAndZ.Add( IAndZ );
+				VertIndexAndZ.Add(IAndZ);
 			}
 		}
 
@@ -4574,7 +4577,7 @@ public:
 	{
 		if (IsInGameThread())
 		{
-			GWarn->BeginSlowTask( NSLOCTEXT("UnrealEd", "ProcessingSkeletalTriangles", "Processing Mesh Triangles"), true );
+			GWarn->BeginSlowTask(NSLOCTEXT("UnrealEd", "ProcessingSkeletalTriangles", "Processing Mesh Triangles"), true);
 		}
 	}
 
@@ -4582,7 +4585,7 @@ public:
 	{
 		if (IsInGameThread())
 		{
-			GWarn->StatusUpdate(Numerator, Denominator, NSLOCTEXT("UnrealEd","ProcessingSkeletalTriangles","Processing Mesh Triangles"));
+			GWarn->StatusUpdate(Numerator, Denominator, NSLOCTEXT("UnrealEd", "ProcessingSkeletalTriangles", "Processing Mesh Triangles"));
 		}
 	}
 
@@ -4644,7 +4647,7 @@ bool FMeshUtilities::BuildSkeletalMesh(FStaticLODModel& LODModel, const FReferen
 	Builder.EndSlowTask();
 
 	// Only show these warnings if in the game thread.  When importing morph targets, this function can run in another thread and these warnings dont prevent the mesh from importing
-	if( IsInGameThread() )
+	if (IsInGameThread())
 	{
 		bool bHasBadSections = false;
 		for (int32 SectionIndex = 0; SectionIndex < BuildData.LODModel.Sections.Num(); SectionIndex++)
@@ -4660,20 +4663,20 @@ bool FMeshUtilities::BuildSkeletalMesh(FStaticLODModel& LODModel, const FReferen
 				Section.NumTriangles
 				);
 		}
-		if( bHasBadSections )
+		if (bHasBadSections)
 		{
 			FText BadSectionMessage(NSLOCTEXT("UnrealEd", "Error_SkeletalMeshHasBadSections", "Input mesh has a section with no triangles.  This mesh may not render properly."));
-			if(BuildData.OutWarningMessages)
+			if (BuildData.OutWarningMessages)
 			{
 				BuildData.OutWarningMessages->Add(BadSectionMessage);
-				if(BuildData.OutWarningNames)
+				if (BuildData.OutWarningNames)
 				{
 					BuildData.OutWarningNames->Add(FFbxErrors::SkeletalMesh_SectionWithNoTriangle);
 				}
 			}
 			else
 			{
-				FMessageDialog::Open( EAppMsgType::Ok, BadSectionMessage );
+				FMessageDialog::Open(EAppMsgType::Ok, BadSectionMessage);
 			}
 		}
 
@@ -4681,10 +4684,10 @@ bool FMeshUtilities::BuildSkeletalMesh(FStaticLODModel& LODModel, const FReferen
 		{
 			FText TooManyVertsMessage(NSLOCTEXT("UnrealEd", "Error_SkeletalMeshTooManyVertices", "Input mesh has too many vertices.  The generated mesh will be corrupt!  Consider adding extra materials to split up the source mesh into smaller chunks."));
 
-			if(BuildData.OutWarningMessages)
+			if (BuildData.OutWarningMessages)
 			{
 				BuildData.OutWarningMessages->Add(TooManyVertsMessage);
-				if(BuildData.OutWarningNames)
+				if (BuildData.OutWarningNames)
 				{
 					BuildData.OutWarningNames->Add(FFbxErrors::SkeletalMesh_TooManyVertices);
 				}
@@ -4698,20 +4701,20 @@ bool FMeshUtilities::BuildSkeletalMesh(FStaticLODModel& LODModel, const FReferen
 
 	return true;
 #else
-	if(OutWarningMessages)
+	if (OutWarningMessages)
 	{
 		OutWarningMessages->Add(FText::FromString(TEXT("Cannot call FMeshUtilities::BuildSkeletalMesh on a console!")));
 	}
 	else
 	{
-		UE_LOG(LogSkeletalMesh, Fatal,TEXT("Cannot call FMeshUtilities::BuildSkeletalMesh on a console!"));
+		UE_LOG(LogSkeletalMesh, Fatal, TEXT("Cannot call FMeshUtilities::BuildSkeletalMesh on a console!"));
 	}
 	return false;
 #endif
 }
 
 //@TODO: The OutMessages has to be a struct that contains FText/FName, or make it Token and add that as error. Needs re-work. Temporary workaround for now. 
-bool FMeshUtilities::BuildSkeletalMesh_Legacy( FStaticLODModel& LODModel, const FReferenceSkeleton& RefSkeleton, const TArray<FVertInfluence>& Influences, const TArray<FMeshWedge>& Wedges, const TArray<FMeshFace>& Faces, const TArray<FVector>& Points, const TArray<int32>& PointToOriginalMap, bool bKeepOverlappingVertices, bool bComputeNormals, bool bComputeTangents, TArray<FText> * OutWarningMessages, TArray<FName> * OutWarningNames)
+bool FMeshUtilities::BuildSkeletalMesh_Legacy(FStaticLODModel& LODModel, const FReferenceSkeleton& RefSkeleton, const TArray<FVertInfluence>& Influences, const TArray<FMeshWedge>& Wedges, const TArray<FMeshFace>& Faces, const TArray<FVector>& Points, const TArray<int32>& PointToOriginalMap, bool bKeepOverlappingVertices, bool bComputeNormals, bool bComputeTangents, TArray<FText> * OutWarningMessages, TArray<FName> * OutWarningNames)
 {
 	bool bTooManyVerts = false;
 
@@ -4723,19 +4726,19 @@ bool FMeshUtilities::BuildSkeletalMesh_Legacy( FStaticLODModel& LODModel, const
 	FaceTangentX.AddUninitialized(Faces.Num());
 	FaceTangentY.AddUninitialized(Faces.Num());
 
-	if( bComputeNormals || bComputeTangents )
+	if (bComputeNormals || bComputeTangents)
 	{
-		for(int32 FaceIndex = 0;FaceIndex < Faces.Num();FaceIndex++)
+		for (int32 FaceIndex = 0; FaceIndex < Faces.Num(); FaceIndex++)
 		{
 			FVector	P1 = Points[Wedges[Faces[FaceIndex].iWedge[0]].iVertex],
 				P2 = Points[Wedges[Faces[FaceIndex].iWedge[1]].iVertex],
 				P3 = Points[Wedges[Faces[FaceIndex].iWedge[2]].iVertex];
-			FVector	TriangleNormal = FPlane(P3,P2,P1);
+			FVector	TriangleNormal = FPlane(P3, P2, P1);
 			FMatrix	ParameterToLocal(
-				FPlane(	P2.X - P1.X,	P2.Y - P1.Y,	P2.Z - P1.Z,	0	),
-				FPlane(	P3.X - P1.X,	P3.Y - P1.Y,	P3.Z - P1.Z,	0	),
-				FPlane(	P1.X,			P1.Y,			P1.Z,			0	),
-				FPlane(	0,				0,				0,				1	)
+				FPlane(P2.X - P1.X, P2.Y - P1.Y, P2.Z - P1.Z, 0),
+				FPlane(P3.X - P1.X, P3.Y - P1.Y, P3.Z - P1.Z, 0),
+				FPlane(P1.X, P1.Y, P1.Z, 0),
+				FPlane(0, 0, 0, 1)
 				);
 
 			float	U1 = Wedges[Faces[FaceIndex].iWedge[0]].UVs[0].X,
@@ -4746,15 +4749,15 @@ bool FMeshUtilities::BuildSkeletalMesh_Legacy( FStaticLODModel& LODModel, const
 				V3 = Wedges[Faces[FaceIndex].iWedge[2]].UVs[0].Y;
 
 			FMatrix	ParameterToTexture(
-				FPlane(	U2 - U1,	V2 - V1,	0,	0	),
-				FPlane(	U3 - U1,	V3 - V1,	0,	0	),
-				FPlane(	U1,			V1,			1,	0	),
-				FPlane(	0,			0,			0,	1	)
+				FPlane(U2 - U1, V2 - V1, 0, 0),
+				FPlane(U3 - U1, V3 - V1, 0, 0),
+				FPlane(U1, V1, 1, 0),
+				FPlane(0, 0, 0, 1)
 				);
 
 			FMatrix	TextureToLocal = ParameterToTexture.Inverse() * ParameterToLocal;
-			FVector	TangentX = TextureToLocal.TransformVector(FVector(1,0,0)).GetSafeNormal(),
-				TangentY = TextureToLocal.TransformVector(FVector(0,1,0)).GetSafeNormal(),
+			FVector	TangentX = TextureToLocal.TransformVector(FVector(1, 0, 0)).GetSafeNormal(),
+				TangentY = TextureToLocal.TransformVector(FVector(0, 1, 0)).GetSafeNormal(),
 				TangentZ;
 
 			TangentX = TangentX - TriangleNormal * (TangentX | TriangleNormal);
@@ -4768,24 +4771,24 @@ bool FMeshUtilities::BuildSkeletalMesh_Legacy( FStaticLODModel& LODModel, const
 	TArray<int32>	WedgeInfluenceIndices;
 
 	// Find wedge influences.
-	TMap<uint32,uint32> VertexIndexToInfluenceIndexMap;
+	TMap<uint32, uint32> VertexIndexToInfluenceIndexMap;
 
-	for(uint32 LookIdx = 0;LookIdx < (uint32)Influences.Num();LookIdx++)
+	for (uint32 LookIdx = 0; LookIdx < (uint32)Influences.Num(); LookIdx++)
 	{
 		// Order matters do not allow the map to overwrite an existing value.
-		if( !VertexIndexToInfluenceIndexMap.Find( Influences[LookIdx].VertIndex) )
+		if (!VertexIndexToInfluenceIndexMap.Find(Influences[LookIdx].VertIndex))
 		{
-			VertexIndexToInfluenceIndexMap.Add( Influences[LookIdx].VertIndex, LookIdx );
+			VertexIndexToInfluenceIndexMap.Add(Influences[LookIdx].VertIndex, LookIdx);
 		}
 	}
 
-	for(int32 WedgeIndex = 0;WedgeIndex < Wedges.Num();WedgeIndex++)
+	for (int32 WedgeIndex = 0; WedgeIndex < Wedges.Num(); WedgeIndex++)
 	{
-		uint32* InfluenceIndex = VertexIndexToInfluenceIndexMap.Find( Wedges[WedgeIndex].iVertex );
+		uint32* InfluenceIndex = VertexIndexToInfluenceIndexMap.Find(Wedges[WedgeIndex].iVertex);
 
-		if ( InfluenceIndex )
+		if (InfluenceIndex)
 		{
-			WedgeInfluenceIndices.Add( *InfluenceIndex );
+			WedgeInfluenceIndices.Add(*InfluenceIndex);
 		}
 		else
 		{
@@ -4808,17 +4811,17 @@ bool FMeshUtilities::BuildSkeletalMesh_Legacy( FStaticLODModel& LODModel, const
 
 	// Calculate smooth wedge tangent vectors.
 
-	if( IsInGameThread() )
+	if (IsInGameThread())
 	{
 		// Only update status if in the game thread.  When importing morph targets, this function can run in another thread
-		GWarn->BeginSlowTask( NSLOCTEXT("UnrealEd", "ProcessingSkeletalTriangles", "Processing Mesh Triangles"), true );
+		GWarn->BeginSlowTask(NSLOCTEXT("UnrealEd", "ProcessingSkeletalTriangles", "Processing Mesh Triangles"), true);
 	}
 
 
 	// To accelerate generation of adjacency, we'll create a table that maps each vertex index
 	// to its overlapping vertices, and a table that maps a vertex to the its influenced faces
-	TMultiMap<int32,int32> Vert2Duplicates;
-	TMultiMap<int32,int32> Vert2Faces;
+	TMultiMap<int32, int32> Vert2Duplicates;
+	TMultiMap<int32, int32> Vert2Faces;
 	TArray<FSkeletalMeshVertIndexAndZ> VertIndexAndZ;
 	{
 		// Create a list of vertex Z/index pairs
@@ -4834,14 +4837,14 @@ bool FMeshUtilities::BuildSkeletalMesh_Legacy( FStaticLODModel& LODModel, const
 		// Sorting function for vertex Z/index pairs
 		struct FCompareFSkeletalMeshVertIndexAndZ
 		{
-			FORCEINLINE bool operator()( const FSkeletalMeshVertIndexAndZ& A, const FSkeletalMeshVertIndexAndZ& B ) const
+			FORCEINLINE bool operator()(const FSkeletalMeshVertIndexAndZ& A, const FSkeletalMeshVertIndexAndZ& B) const
 			{
 				return A.Z < B.Z;
 			}
 		};
 
 		// Sort the vertices by z value
-		VertIndexAndZ.Sort( FCompareFSkeletalMeshVertIndexAndZ() );
+		VertIndexAndZ.Sort(FCompareFSkeletalMeshVertIndexAndZ());
 
 		// Search for duplicates, quickly!
 		for (int32 i = 0; i < VertIndexAndZ.Num(); i++)
@@ -4849,19 +4852,19 @@ bool FMeshUtilities::BuildSkeletalMesh_Legacy( FStaticLODModel& LODModel, const
 			// only need to search forward, since we add pairs both ways
 			for (int32 j = i + 1; j < VertIndexAndZ.Num(); j++)
 			{
-				if (FMath::Abs(VertIndexAndZ[j].Z - VertIndexAndZ[i].Z) > THRESH_POINTS_ARE_SAME )
+				if (FMath::Abs(VertIndexAndZ[j].Z - VertIndexAndZ[i].Z) > THRESH_POINTS_ARE_SAME)
 				{
 					// our list is sorted, so there can't be any more dupes
 					break;
 				}
 
 				// check to see if the points are really overlapping
-				if(PointsEqual(
+				if (PointsEqual(
 					Points[VertIndexAndZ[i].Index],
-					Points[VertIndexAndZ[j].Index] ))					
+					Points[VertIndexAndZ[j].Index]))
 				{
-					Vert2Duplicates.Add(VertIndexAndZ[i].Index,VertIndexAndZ[j].Index);
-					Vert2Duplicates.Add(VertIndexAndZ[j].Index,VertIndexAndZ[i].Index);
+					Vert2Duplicates.Add(VertIndexAndZ[i].Index, VertIndexAndZ[j].Index);
+					Vert2Duplicates.Add(VertIndexAndZ[j].Index, VertIndexAndZ[i].Index);
 				}
 			}
 		}
@@ -4870,12 +4873,12 @@ bool FMeshUtilities::BuildSkeletalMesh_Legacy( FStaticLODModel& LODModel, const
 		VertIndexAndZ.Reset();
 
 		// now create a map from vert indices to faces
-		for(int32 FaceIndex = 0;FaceIndex < Faces.Num();FaceIndex++)
+		for (int32 FaceIndex = 0; FaceIndex < Faces.Num(); FaceIndex++)
 		{
 			const FMeshFace&	Face = Faces[FaceIndex];
-			for(int32 VertexIndex = 0;VertexIndex < 3;VertexIndex++)
+			for (int32 VertexIndex = 0; VertexIndex < 3; VertexIndex++)
 			{
-				Vert2Faces.AddUnique(Wedges[Face.iWedge[VertexIndex]].iVertex,FaceIndex);
+				Vert2Faces.AddUnique(Wedges[Face.iWedge[VertexIndex]].iVertex, FaceIndex);
 			}
 		}
 	}
@@ -4887,29 +4890,29 @@ bool FMeshUtilities::BuildSkeletalMesh_Legacy( FStaticLODModel& LODModel, const
 
 	// List of raw calculated vertices that will be merged later
 	TArray<FSoftSkinBuildVertex> RawVertices;
-	RawVertices.Reserve( Points.Num() );
+	RawVertices.Reserve(Points.Num());
 
 	// Create a list of vertex Z/index pairs
 
-	for(int32 FaceIndex = 0;FaceIndex < Faces.Num();FaceIndex++)
+	for (int32 FaceIndex = 0; FaceIndex < Faces.Num(); FaceIndex++)
 	{
 		// Only update the status progress bar if we are in the gamethread and every thousand faces. 
 		// Updating status is extremely slow
-		if( FaceIndex % 5000 == 0 && IsInGameThread() )
+		if (FaceIndex % 5000 == 0 && IsInGameThread())
 		{
 			// Only update status if in the game thread.  When importing morph targets, this function can run in another thread
-			GWarn->StatusUpdate( FaceIndex, Faces.Num(), NSLOCTEXT("UnrealEd", "ProcessingSkeletalTriangles", "Processing Mesh Triangles") );
+			GWarn->StatusUpdate(FaceIndex, Faces.Num(), NSLOCTEXT("UnrealEd", "ProcessingSkeletalTriangles", "Processing Mesh Triangles"));
 		}
 
 		const FMeshFace&	Face = Faces[FaceIndex];
 
 		FVector	VertexTangentX[3],
-				VertexTangentY[3],
-				VertexTangentZ[3];
+			VertexTangentY[3],
+			VertexTangentZ[3];
 
-		if( bComputeNormals || bComputeTangents )
+		if (bComputeNormals || bComputeTangents)
 		{
-			for(int32 VertexIndex = 0;VertexIndex < 3;VertexIndex++)
+			for (int32 VertexIndex = 0; VertexIndex < 3; VertexIndex++)
 			{
 				VertexTangentX[VertexIndex] = FVector::ZeroVector;
 				VertexTangentY[VertexIndex] = FVector::ZeroVector;
@@ -4920,21 +4923,21 @@ bool FMeshUtilities::BuildSkeletalMesh_Legacy( FStaticLODModel& LODModel, const
 				Points[Wedges[Face.iWedge[2]].iVertex],
 				Points[Wedges[Face.iWedge[1]].iVertex],
 				Points[Wedges[Face.iWedge[0]].iVertex]
-			);
-			float	Determinant = FVector::Triple(FaceTangentX[FaceIndex],FaceTangentY[FaceIndex],TriangleNormal);
+				);
+			float	Determinant = FVector::Triple(FaceTangentX[FaceIndex], FaceTangentY[FaceIndex], TriangleNormal);
 
 			// Start building a list of faces adjacent to this triangle
 			AdjacentFaces.Reset();
-			for(int32 VertexIndex = 0;VertexIndex < 3;VertexIndex++)
+			for (int32 VertexIndex = 0; VertexIndex < 3; VertexIndex++)
 			{
 				int32 vert = Wedges[Face.iWedge[VertexIndex]].iVertex;
 				DupVerts.Reset();
-				Vert2Duplicates.MultiFind(vert,DupVerts);
+				Vert2Duplicates.MultiFind(vert, DupVerts);
 				DupVerts.Add(vert); // I am a "dupe" of myself
 				for (int32 k = 0; k < DupVerts.Num(); k++)
 				{
 					DupFaces.Reset();
-					Vert2Faces.MultiFind(DupVerts[k],DupFaces);
+					Vert2Faces.MultiFind(DupVerts[k], DupFaces);
 					for (int32 l = 0; l < DupFaces.Num(); l++)
 					{
 						AdjacentFaces.AddUnique(DupFaces[l]);
@@ -4943,7 +4946,7 @@ bool FMeshUtilities::BuildSkeletalMesh_Legacy( FStaticLODModel& LODModel, const
 			}
 
 			// Process adjacent faces
-			for(int32 AdjacentFaceIndex = 0;AdjacentFaceIndex < AdjacentFaces.Num();AdjacentFaceIndex++)
+			for (int32 AdjacentFaceIndex = 0; AdjacentFaceIndex < AdjacentFaces.Num(); AdjacentFaceIndex++)
 			{
 				int32 OtherFaceIndex = AdjacentFaces[AdjacentFaceIndex];
 				const FMeshFace&	OtherFace = Faces[OtherFaceIndex];
@@ -4951,19 +4954,19 @@ bool FMeshUtilities::BuildSkeletalMesh_Legacy( FStaticLODModel& LODModel, const
 					Points[Wedges[OtherFace.iWedge[2]].iVertex],
 					Points[Wedges[OtherFace.iWedge[1]].iVertex],
 					Points[Wedges[OtherFace.iWedge[0]].iVertex]
-				);
-				float		OtherFaceDeterminant = FVector::Triple(FaceTangentX[OtherFaceIndex],FaceTangentY[OtherFaceIndex],OtherTriangleNormal);
+					);
+				float		OtherFaceDeterminant = FVector::Triple(FaceTangentX[OtherFaceIndex], FaceTangentY[OtherFaceIndex], OtherTriangleNormal);
 
-				for(int32 VertexIndex = 0;VertexIndex < 3;VertexIndex++)
+				for (int32 VertexIndex = 0; VertexIndex < 3; VertexIndex++)
 				{
-					for(int32 OtherVertexIndex = 0;OtherVertexIndex < 3;OtherVertexIndex++)
+					for (int32 OtherVertexIndex = 0; OtherVertexIndex < 3; OtherVertexIndex++)
 					{
-						if(PointsEqual(
+						if (PointsEqual(
 							Points[Wedges[OtherFace.iWedge[OtherVertexIndex]].iVertex],
 							Points[Wedges[Face.iWedge[VertexIndex]].iVertex]
-						))					
+							))
 						{
-							if(Determinant * OtherFaceDeterminant > 0.0f && SkeletalMeshTools::SkeletalMesh_UVsEqual(Wedges[OtherFace.iWedge[OtherVertexIndex]],Wedges[Face.iWedge[VertexIndex]]))
+							if (Determinant * OtherFaceDeterminant > 0.0f && SkeletalMeshTools::SkeletalMesh_UVsEqual(Wedges[OtherFace.iWedge[OtherVertexIndex]], Wedges[Face.iWedge[VertexIndex]]))
 							{
 								VertexTangentX[VertexIndex] += FaceTangentX[OtherFaceIndex];
 								VertexTangentY[VertexIndex] += FaceTangentY[OtherFaceIndex];
@@ -4971,7 +4974,7 @@ bool FMeshUtilities::BuildSkeletalMesh_Legacy( FStaticLODModel& LODModel, const
 
 							// Only contribute 'normal' if the vertices are truly one and the same to obey hard "smoothing" edges baked into 
 							// the mesh by vertex duplication
-							if( Wedges[OtherFace.iWedge[OtherVertexIndex]].iVertex == Wedges[Face.iWedge[VertexIndex]].iVertex ) 
+							if (Wedges[OtherFace.iWedge[OtherVertexIndex]].iVertex == Wedges[Face.iWedge[VertexIndex]].iVertex)
 							{
 								VertexTangentZ[VertexIndex] += OtherTriangleNormal;
 							}
@@ -4981,20 +4984,20 @@ bool FMeshUtilities::BuildSkeletalMesh_Legacy( FStaticLODModel& LODModel, const
 			}
 		}
 
-		for(int32 VertexIndex = 0;VertexIndex < 3;VertexIndex++)
+		for (int32 VertexIndex = 0; VertexIndex < 3; VertexIndex++)
 		{
 			FSoftSkinBuildVertex	Vertex;
 
 			Vertex.Position = Points[Wedges[Face.iWedge[VertexIndex]].iVertex];
 
-			FVector TangentX,TangentY,TangentZ;
+			FVector TangentX, TangentY, TangentZ;
 
-			if( bComputeNormals || bComputeTangents )
+			if (bComputeNormals || bComputeTangents)
 			{
 				TangentX = VertexTangentX[VertexIndex].GetSafeNormal();
 				TangentY = VertexTangentY[VertexIndex].GetSafeNormal();
 
-				if( bComputeNormals )
+				if (bComputeNormals)
 				{
 					TangentZ = VertexTangentZ[VertexIndex].GetSafeNormal();
 				}
@@ -5028,7 +5031,7 @@ bool FMeshUtilities::BuildSkeletalMesh_Legacy( FStaticLODModel& LODModel, const
 			Vertex.TangentY = TangentY;
 			Vertex.TangentZ = TangentZ;
 
-			FMemory::Memcpy( Vertex.UVs,  Wedges[Face.iWedge[VertexIndex]].UVs, sizeof(FVector2D)*MAX_TEXCOORDS);	
+			FMemory::Memcpy(Vertex.UVs, Wedges[Face.iWedge[VertexIndex]].UVs, sizeof(FVector2D)*MAX_TEXCOORDS);
 			Vertex.Color = Wedges[Face.iWedge[VertexIndex]].Color;
 
 			{
@@ -5038,32 +5041,32 @@ bool FMeshUtilities::BuildSkeletalMesh_Legacy( FStaticLODModel& LODModel, const
 				int32 LookIdx = InfIdx;
 
 				uint32 InfluenceCount = 0;
-				while( Influences.IsValidIndex(LookIdx) && (Influences[LookIdx].VertIndex == Wedges[Face.iWedge[VertexIndex]].iVertex) )
-				{			
+				while (Influences.IsValidIndex(LookIdx) && (Influences[LookIdx].VertIndex == Wedges[Face.iWedge[VertexIndex]].iVertex))
+				{
 					InfluenceCount++;
 					LookIdx++;
 				}
-				InfluenceCount = FMath::Min<uint32>(InfluenceCount,MAX_TOTAL_INFLUENCES);
+				InfluenceCount = FMath::Min<uint32>(InfluenceCount, MAX_TOTAL_INFLUENCES);
 
 				// Setup the vertex influences.
 
 				Vertex.InfluenceBones[0] = 0;
 				Vertex.InfluenceWeights[0] = 255;
-				for(uint32 i = 1;i < MAX_TOTAL_INFLUENCES;i++)
+				for (uint32 i = 1; i < MAX_TOTAL_INFLUENCES; i++)
 				{
 					Vertex.InfluenceBones[i] = 0;
 					Vertex.InfluenceWeights[i] = 0;
 				}
 
 				uint32	TotalInfluenceWeight = 0;
-				for(uint32 i = 0;i < InfluenceCount;i++)
+				for (uint32 i = 0; i < InfluenceCount; i++)
 				{
-					FBoneIndexType BoneIndex = (FBoneIndexType)Influences[InfIdx+i].BoneIndex;
-					if( BoneIndex >= RefSkeleton.GetNum() )
+					FBoneIndexType BoneIndex = (FBoneIndexType)Influences[InfIdx + i].BoneIndex;
+					if (BoneIndex >= RefSkeleton.GetNum())
 						continue;
 
 					Vertex.InfluenceBones[i] = BoneIndex;
-					Vertex.InfluenceWeights[i] = (uint8)(Influences[InfIdx+i].Weight * 255.0f);
+					Vertex.InfluenceWeights[i] = (uint8)(Influences[InfIdx + i].Weight * 255.0f);
 					TotalInfluenceWeight += Vertex.InfluenceWeights[i];
 				}
 				Vertex.InfluenceWeights[0] += 255 - TotalInfluenceWeight;
@@ -5071,37 +5074,37 @@ bool FMeshUtilities::BuildSkeletalMesh_Legacy( FStaticLODModel& LODModel, const
 
 			// Add the vertex as well as its original index in the points array
 			Vertex.PointWedgeIdx = Wedges[Face.iWedge[VertexIndex]].iVertex;
-			
-			int32 RawIndex = RawVertices.Add( Vertex );
+
+			int32 RawIndex = RawVertices.Add(Vertex);
 
 			// Add an efficient way to find dupes of this vertex later for fast combining of vertices
 			FSkeletalMeshVertIndexAndZ IAndZ;
 			IAndZ.Index = RawIndex;
 			IAndZ.Z = Vertex.Position.Z;
 
-			VertIndexAndZ.Add( IAndZ );
+			VertIndexAndZ.Add(IAndZ);
 		}
 	}
 
 	// Generate chunks and their vertices and indices
-	SkeletalMeshTools::BuildSkeletalMeshChunks( Faces, RawVertices, VertIndexAndZ, bKeepOverlappingVertices, Chunks, bTooManyVerts );
+	SkeletalMeshTools::BuildSkeletalMeshChunks(Faces, RawVertices, VertIndexAndZ, bKeepOverlappingVertices, Chunks, bTooManyVerts);
 
 	// Chunk vertices to satisfy the requested limit.
 	static const auto MaxBonesVar = IConsoleManager::Get().FindTConsoleVariableDataInt(TEXT("Compat.MAX_GPUSKIN_BONES"));
 	const int32 MaxGPUSkinBones = MaxBonesVar->GetValueOnAnyThread();
-	SkeletalMeshTools::ChunkSkinnedVertices(Chunks,MaxGPUSkinBones);
+	SkeletalMeshTools::ChunkSkinnedVertices(Chunks, MaxGPUSkinBones);
 
 	// Build the skeletal model from chunks.
-	BuildSkeletalModelFromChunks(LODModel,RefSkeleton,Chunks,PointToOriginalMap);
+	BuildSkeletalModelFromChunks(LODModel, RefSkeleton, Chunks, PointToOriginalMap);
 
-	if( IsInGameThread() )
+	if (IsInGameThread())
 	{
 		// Only update status if in the game thread.  When importing morph targets, this function can run in another thread
 		GWarn->EndSlowTask();
 	}
 
 	// Only show these warnings if in the game thread.  When importing morph targets, this function can run in another thread and these warnings dont prevent the mesh from importing
-	if( IsInGameThread() )
+	if (IsInGameThread())
 	{
 		bool bHasBadSections = false;
 		for (int32 SectionIndex = 0; SectionIndex < LODModel.Sections.Num(); SectionIndex++)
@@ -5117,20 +5120,20 @@ bool FMeshUtilities::BuildSkeletalMesh_Legacy( FStaticLODModel& LODModel, const
 				Section.NumTriangles
 				);
 		}
-		if( bHasBadSections )
+		if (bHasBadSections)
 		{
 			FText BadSectionMessage(NSLOCTEXT("UnrealEd", "Error_SkeletalMeshHasBadSections", "Input mesh has a section with no triangles.  This mesh may not render properly."));
-			if(OutWarningMessages)
+			if (OutWarningMessages)
 			{
 				OutWarningMessages->Add(BadSectionMessage);
-				if(OutWarningNames)
+				if (OutWarningNames)
 				{
 					OutWarningNames->Add(FFbxErrors::SkeletalMesh_SectionWithNoTriangle);
 				}
 			}
 			else
 			{
-				FMessageDialog::Open( EAppMsgType::Ok, BadSectionMessage );
+				FMessageDialog::Open(EAppMsgType::Ok, BadSectionMessage);
 			}
 		}
 
@@ -5138,10 +5141,10 @@ bool FMeshUtilities::BuildSkeletalMesh_Legacy( FStaticLODModel& LODModel, const
 		{
 			FText TooManyVertsMessage(NSLOCTEXT("UnrealEd", "Error_SkeletalMeshTooManyVertices", "Input mesh has too many vertices.  The generated mesh will be corrupt!  Consider adding extra materials to split up the source mesh into smaller chunks."));
 
-			if(OutWarningMessages)
+			if (OutWarningMessages)
 			{
 				OutWarningMessages->Add(TooManyVertsMessage);
-				if(OutWarningNames)
+				if (OutWarningNames)
 				{
 					OutWarningNames->Add(FFbxErrors::SkeletalMesh_TooManyVertices);
 				}
@@ -5168,7 +5171,7 @@ static bool NonOpaqueMaterialPredicate(UStaticMeshComponent* InMesh)
 		}
 	}
 
-	return false; 
+	return false;
 }
 
 static FIntPoint ConditionalImageResize(const FIntPoint& SrcSize, const FIntPoint& DesiredSize, TArray<FColor>& InOutImage, bool bLinearSpace)
@@ -5226,10 +5229,10 @@ void FMeshUtilities::CreateProxyMesh(const TArray<AActor*>& InActors, const stru
 			TInlineComponentArray<UStaticMeshComponent*> Components;
 			Actor->GetComponents<UStaticMeshComponent>(Components);
 			// TODO: support derived classes from static component
-			Components.RemoveAll([](UStaticMeshComponent* Val){ return !(Val->IsA(UStaticMeshComponent::StaticClass()) || Val->IsA(USplineMeshComponent::StaticClass())); });		
-			
+			Components.RemoveAll([](UStaticMeshComponent* Val){ return !(Val->IsA(UStaticMeshComponent::StaticClass()) || Val->IsA(USplineMeshComponent::StaticClass())); });
+
 			// TODO: support non-opaque materials
-			Components.RemoveAll(&NonOpaqueMaterialPredicate);
+			//Components.RemoveAll(&NonOpaqueMaterialPredicate);
 			//
 			ComponentsToMerge.Append(Components);
 
@@ -5267,7 +5270,7 @@ void FMeshUtilities::CreateProxyMesh(const TArray<AActor*>& InActors, const stru
 			// Propagate painted vertex colors into our raw mesh
 			PropagatePaintedColorsToRawMesh(StaticMeshComponent, 0, RawMesh);
 			// Whether at least one of the meshes has vertex colors
-			bVertexColorsInMeshes |= (RawMesh.WedgeColors.Num() != 0);			
+			bVertexColorsInMeshes |= (RawMesh.WedgeColors.Num() != 0);
 
 			// Which UV channels has data at least in one mesh
 			for (int32 ChannelIdx = 0; ChannelIdx < MAX_MESH_TEXTURE_COORDS; ++ChannelIdx)
@@ -5306,7 +5309,7 @@ void FMeshUtilities::CreateProxyMesh(const TArray<AActor*>& InActors, const stru
 	//For each raw mesh, re-map the material indices from Local to Global material indices space
 	for (int32 RawMeshIndex = 0; RawMeshIndex < SourceMeshes.Num(); ++RawMeshIndex)
 	{
-		const TArray<int32>& GlobalMaterialIndices = *MaterialMap.Find(FMeshIdAndLOD(RawMeshIndex,0));
+		const TArray<int32>& GlobalMaterialIndices = *MaterialMap.Find(FMeshIdAndLOD(RawMeshIndex, 0));
 		TArray<int32>& MaterialIndices = SourceMeshes[RawMeshIndex].MeshLODData[0].RawMesh.FaceMaterialIndices;
 		int32 MaterialIndicesCount = MaterialIndices.Num();
 
@@ -5332,7 +5335,7 @@ void FMeshUtilities::CreateProxyMesh(const TArray<AActor*>& InActors, const stru
 
 	// Add this proxy job to map	
 	Processor->AddProxyJob(InGuid, Data);
-	
+
 	// We are only using LOD level 0
 	TArray<FMeshMergeData> MergeData;
 	for (FRawMeshExt& SourceMesh : SourceMeshes)
@@ -5363,14 +5366,14 @@ void FMeshUtilities::CreateProxyMesh(const TArray<AActor*>& Actors, const struct
 	FGuid JobGuid = FGuid::NewGuid();
 	Delegate.BindLambda(
 		[&](const FGuid Guid, TArray<UObject*>& InAssetsToSync)
+	{
+		if (JobGuid == Guid)
 		{
-			if (JobGuid == Guid)
-			{
-				OutAssetsToSync.Append(InAssetsToSync);
-			}
+			OutAssetsToSync.Append(InAssetsToSync);
 		}
+	}
 	);
-	
+
 	CreateProxyMesh(Actors, InProxySettings, InOuter, ProxyBasePackageName, JobGuid, Delegate, false);
 }
 
@@ -5393,13 +5396,15 @@ void FMeshUtilities::FlattenMaterialsWithMeshData(TArray<UMaterialInterface*>& I
 
 		FFlattenMaterial FlattenMaterial = FMaterialUtilities::CreateFlattenMaterialWithSettings(InMaterialProxySettings);
 
-		/* Find a mesh which uses the current material. Materials using vertex data are added for each individual mesh using it, 
+		/* Find a mesh which uses the current material. Materials using vertex data are added for each individual mesh using it,
 		which is why baking down the materials like this works. :) */
 		int32 UsedMeshIndex = 0;
 		int32 LocalMaterialIndex = 0;
+		int32 LocalTextureBoundIndex = 0;
 		FMeshMergeData* MergeData = nullptr;
 		for (int32 MeshIndex = 0; MeshIndex < InSourceMeshes.Num() && MergeData == nullptr; MeshIndex++)
 		{
+			bool bFound = false;
 			for (int32 LODIndex = 0; LODIndex < MAX_STATIC_MESH_LODS; ++LODIndex)
 			{
 				if (InSourceMeshes[MeshIndex].MeshLODData[LODIndex].RawMesh.VertexPositions.Num())
@@ -5411,6 +5416,8 @@ void FMeshUtilities::FlattenMaterialsWithMeshData(TArray<UMaterialInterface*>& I
 						{
 							UsedMeshIndex = MeshIndex;
 							MergeData = &InSourceMeshes[MeshIndex].MeshLODData[LODIndex];
+							LocalTextureBoundIndex = LocalMaterialIndex;
+							bFound = true;
 							break;
 						}
 					}
@@ -5418,10 +5425,16 @@ void FMeshUtilities::FlattenMaterialsWithMeshData(TArray<UMaterialInterface*>& I
 				else
 				{
 					break;
-				}			
+				}
+			}
+
+			// Early out if we found the mesh we need
+			if (bFound == true)
+			{
+				break;
 			}
 		}
-		
+
 		// If there is specific vertex data available and used in the material we should generate non-overlapping UVs
 		if (MergeData && InMeshShouldBakeVertexData[UsedMeshIndex])
 		{
@@ -5434,12 +5447,13 @@ void FMeshUtilities::FlattenMaterialsWithMeshData(TArray<UMaterialInterface*>& I
 				// Generate unique UVs
 				GenerateUniqueUVsForStaticMesh(MergeData->RawMesh, InMaterialProxySettings.TextureSize.GetMax(), MergeData->NewUVs);
 			}
+
 			// Export the material using mesh data to support vertex based material properties
 			FMaterialUtilities::ExportMaterial(
 				CurrentMaterial,
 				&MergeData->RawMesh,
 				LocalMaterialIndex,
-				MergeData->TexCoordBounds[LocalMaterialIndex],
+				MergeData->TexCoordBounds[LocalTextureBoundIndex],
 				MergeData->NewUVs,
 				FlattenMaterial,
 				CachedShader);
@@ -5455,7 +5469,7 @@ void FMeshUtilities::FlattenMaterialsWithMeshData(TArray<UMaterialInterface*>& I
 
 		// Fill flatten material samples alpha values with 255 (for saving out textures correctly for Simplygon Swarm)
 		FlattenMaterial.FillAlphaValues(255);
-		
+
 		// Add flattened material to outgoing array
 		OutFlattenedMaterials.Add(FlattenMaterial);
 
@@ -5523,11 +5537,11 @@ static void ExportStaticMeshLOD(const FStaticMeshLODResources& StaticMeshLOD, FR
 {
 	const int32 NumWedges = StaticMeshLOD.IndexBuffer.GetNumIndices();
 	const int32 NumVertexPositions = StaticMeshLOD.PositionVertexBuffer.GetNumVertices();
-	const int32 NumFaces = NumWedges/3;
+	const int32 NumFaces = NumWedges / 3;
 
 	// Indices
 	StaticMeshLOD.IndexBuffer.GetCopy(OutRawMesh.WedgeIndices);
-		
+
 	// Vertex positions
 	if (NumVertexPositions > 0)
 	{
@@ -5538,7 +5552,7 @@ static void ExportStaticMeshLOD(const FStaticMeshLODResources& StaticMeshLOD, FR
 			OutRawMesh.VertexPositions.Add(Pos);
 		}
 	}
-		
+
 	// Vertex data
 	if (StaticMeshLOD.VertexBuffer.GetNumVertices() > 0)
 	{
@@ -5551,7 +5565,7 @@ static void ExportStaticMeshLOD(const FStaticMeshLODResources& StaticMeshLOD, FR
 		{
 			OutRawMesh.WedgeTexCoords[TexCoodIdx].Empty(NumWedges);
 		}
-		
+
 		for (int32 WedgeIndex : OutRawMesh.WedgeIndices)
 		{
 			FVector WedgeTangentX = StaticMeshLOD.VertexBuffer.VertexTangentX(WedgeIndex);
@@ -5587,7 +5601,7 @@ static void ExportStaticMeshLOD(const FStaticMeshLODResources& StaticMeshLOD, FR
 
 		for (const FStaticMeshSection& Section : StaticMeshLOD.Sections)
 		{
-			uint32 FirstTriangle = Section.FirstIndex/3;
+			uint32 FirstTriangle = Section.FirstIndex / 3;
 			for (uint32 TriangleIndex = 0; TriangleIndex < Section.NumTriangles; ++TriangleIndex)
 			{
 				OutRawMesh.FaceMaterialIndices[FirstTriangle + TriangleIndex] = Section.MaterialIndex;
@@ -5610,7 +5624,7 @@ static void ExportStaticMeshLOD(const FStaticMeshLODResources& StaticMeshLOD, FR
 bool FMeshUtilities::ConstructRawMesh(
 	UStaticMeshComponent* InMeshComponent,
 	int32 LODIndex,
-	FRawMesh& OutRawMesh, 
+	FRawMesh& OutRawMesh,
 	TArray<UMaterialInterface*>& OutUniqueMaterials,
 	TArray<int32>& OutGlobalMaterialIndices) const
 {
@@ -5620,7 +5634,7 @@ bool FMeshUtilities::ConstructRawMesh(
 		UE_LOG(LogMeshUtilities, Warning, TEXT("No static mesh actor found in component %s."), *InMeshComponent->GetName());
 		return false;
 	}
-	
+
 	if (!SrcMesh->SourceModels.IsValidIndex(LODIndex))
 	{
 		UE_LOG(LogMeshUtilities, Log, TEXT("No mesh data found for LOD%d %s."), LODIndex, *SrcMesh->GetName());
@@ -5632,11 +5646,11 @@ bool FMeshUtilities::ConstructRawMesh(
 		UE_LOG(LogMeshUtilities, Warning, TEXT("No mesh render data found for LOD%d %s."), LODIndex, *SrcMesh->GetName());
 		return false;
 	}
-	
+
 	FStaticMeshSourceModel& SrcModel = SrcMesh->SourceModels[LODIndex];
 	const bool bImportedMesh = !SrcModel.RawMeshBulkData->IsEmpty();
 
-	if (bImportedMesh)
+	if (bImportedMesh && !InMeshComponent->IsA<USplineMeshComponent>())
 	{
 		SrcModel.RawMeshBulkData->LoadRawMesh(OutRawMesh);
 	}
@@ -5645,14 +5659,14 @@ bool FMeshUtilities::ConstructRawMesh(
 		// Reconstruct RawMesh from render data
 		ExportStaticMeshLOD(SrcMesh->RenderData->LODResources[LODIndex], OutRawMesh);
 	}
-	
+
 	// Make sure the raw mesh is not irreparably malformed.
 	if (!OutRawMesh.IsValidOrFixable())
 	{
 		UE_LOG(LogMeshUtilities, Error, TEXT("Raw mesh (%s) is corrupt for LOD%d."), *SrcMesh->GetName(), LODIndex);
 		return false;
 	}
-	
+
 	// Handle spline mesh deformation
 	if (InMeshComponent->IsA<USplineMeshComponent>())
 	{
@@ -5676,7 +5690,7 @@ bool FMeshUtilities::ConstructRawMesh(
 			if (OutRawMesh.WedgeTangentZ.Num())
 			{
 				OutRawMesh.WedgeTangentZ[iVert] = SliceTransform.TransformVector(OutRawMesh.WedgeTangentZ[iVert]);
-			}			
+			}
 		}
 
 		for (int32 iVert = 0; iVert < OutRawMesh.VertexPositions.Num(); ++iVert)
@@ -5716,17 +5730,17 @@ bool FMeshUtilities::ConstructRawMesh(
 	{
 		TangentZ = CtoM.TransformVectorNoScale(TangentZ);
 	}
-
+	
 	const bool bIsMirrored = CtoM.GetDeterminant() < 0.f;
 	if (bIsMirrored)
 	{
 		// Flip faces
-		for (int32 FaceIdx = 0; FaceIdx < OutRawMesh.WedgeIndices.Num()/3; FaceIdx++)
+		for (int32 FaceIdx = 0; FaceIdx < OutRawMesh.WedgeIndices.Num() / 3; FaceIdx++)
 		{
 			int32 I0 = FaceIdx * 3 + 0;
 			int32 I2 = FaceIdx * 3 + 2;
 			Swap(OutRawMesh.WedgeIndices[I0], OutRawMesh.WedgeIndices[I2]);
-			
+
 			// seems like vertex colors and UVs are not indexed, so swap values instead
 			if (OutRawMesh.WedgeColors.Num())
 			{
@@ -5742,12 +5756,12 @@ bool FMeshUtilities::ConstructRawMesh(
 			}
 		}
 	}
-		
+
 	int32 NumWedges = OutRawMesh.WedgeIndices.Num();
 
 	// Figure out if we should recompute normals and tangents. By default generated LODs should not recompute normals
-	bool bRecomputeNormals = (bImportedMesh && BuildSettings.bRecomputeNormals) || OutRawMesh.WedgeTangentZ.Num() == 0;
-	bool bRecomputeTangents = (bImportedMesh && BuildSettings.bRecomputeTangents) || OutRawMesh.WedgeTangentX.Num() == 0 || OutRawMesh.WedgeTangentY.Num() == 0;
+	bool bRecomputeNormals = (bImportedMesh && BuildSettings.bRecomputeNormals) || OutRawMesh.WedgeTangentZ.Num() == 0 || bIsMirrored;
+	bool bRecomputeTangents = (bImportedMesh && BuildSettings.bRecomputeTangents) || OutRawMesh.WedgeTangentX.Num() == 0 || OutRawMesh.WedgeTangentY.Num() == 0 || bIsMirrored;
 
 	// Dump normals and tangents if we are recomputing them.
 	if (bRecomputeTangents)
@@ -5757,7 +5771,7 @@ bool FMeshUtilities::ConstructRawMesh(
 		OutRawMesh.WedgeTangentY.Empty(NumWedges);
 		OutRawMesh.WedgeTangentY.AddZeroed(NumWedges);
 	}
-	
+
 	if (bRecomputeNormals)
 	{
 		OutRawMesh.WedgeTangentZ.Empty(NumWedges);
@@ -5768,7 +5782,7 @@ bool FMeshUtilities::ConstructRawMesh(
 	if (bRecomputeNormals || bRecomputeTangents)
 	{
 		float ComparisonThreshold = GetComparisonThreshold(BuildSettings);
-		TMultiMap<int32,int32> OverlappingCorners;
+		TMultiMap<int32, int32> OverlappingCorners;
 		FindOverlappingCorners(OverlappingCorners, OutRawMesh, ComparisonThreshold);
 
 		// Static meshes always blend normals of overlapping corners.
@@ -5797,7 +5811,7 @@ bool FMeshUtilities::ConstructRawMesh(
 
 	//Need to store the unique material indices in order to re-map the material indices in each rawmesh
 	//Only using the base mesh
-	for (const FStaticMeshSection& Section : SrcMesh->RenderData->LODResources[LODIndex].Sections) 
+	for (const FStaticMeshSection& Section : SrcMesh->RenderData->LODResources[LODIndex].Sections)
 	{
 		// Add material and store the material ID
 		UMaterialInterface* MaterialToAdd = InMeshComponent->GetMaterial(Section.MaterialIndex);
@@ -5815,11 +5829,11 @@ bool FMeshUtilities::ConstructRawMesh(
 		{
 			MaterialToAdd = DefaultMaterial;
 		}
-		
+
 		const int32 MaterialIdx = OutUniqueMaterials.Add(MaterialToAdd);
 		OutGlobalMaterialIndices.Add(MaterialIdx);
 	}
-			
+
 	return true;
 }
 
@@ -5850,7 +5864,7 @@ void FMeshUtilities::ExtractMeshDataForGeometryCache(FRawMesh& RawMesh, const FM
 	TMultiMap<int32, int32> OverlappingCorners;
 	if (bRecomputeNormals || bRecomputeTangents)
 	{
-		float ComparisonThreshold = GetComparisonThreshold(BuildSettings);		
+		float ComparisonThreshold = GetComparisonThreshold(BuildSettings);
 		FindOverlappingCorners(OverlappingCorners, RawMesh, ComparisonThreshold);
 
 		// Static meshes always blend normals of overlapping corners.
@@ -5874,7 +5888,7 @@ void FMeshUtilities::ExtractMeshDataForGeometryCache(FRawMesh& RawMesh, const FM
 	check(RawMesh.WedgeTangentX.Num() == NumWedges);
 	check(RawMesh.WedgeTangentY.Num() == NumWedges);
 	check(RawMesh.WedgeTangentZ.Num() == NumWedges);
-		
+
 	TArray<int32> OutWedgeMap;
 
 	int32 MaxMaterialIndex = 1;
@@ -5885,7 +5899,7 @@ void FMeshUtilities::ExtractMeshDataForGeometryCache(FRawMesh& RawMesh, const FM
 
 	for (int32 i = 0; i <= MaxMaterialIndex; ++i)
 	{
-		OutPerSectionIndices.Push( TArray<uint32>() );
+		OutPerSectionIndices.Push(TArray<uint32>());
 	}
 
 	BuildStaticMeshVertexAndIndexBuffers(OutVertices, OutPerSectionIndices, OutWedgeMap, RawMesh, OverlappingCorners, KINDA_SMALL_NUMBER, BuildSettings.BuildScale3D);
@@ -5898,22 +5912,22 @@ void FMeshUtilities::ExtractMeshDataForGeometryCache(FRawMesh& RawMesh, const FM
 }
 
 /*------------------------------------------------------------------------------
-	Mesh merging 
+Mesh merging
 ------------------------------------------------------------------------------*/
 bool FMeshUtilities::PropagatePaintedColorsToRawMesh(UStaticMeshComponent* StaticMeshComponent, int32 LODIndex, FRawMesh& RawMesh) const
 {
 	UStaticMesh* StaticMesh = StaticMeshComponent->StaticMesh;
-	
+
 	if (StaticMesh->SourceModels.IsValidIndex(LODIndex) &&
 		StaticMeshComponent->LODData.IsValidIndex(LODIndex) &&
-		StaticMeshComponent->LODData[LODIndex].OverrideVertexColors != nullptr)	
+		StaticMeshComponent->LODData[LODIndex].OverrideVertexColors != nullptr)
 	{
 		FColorVertexBuffer& ColorVertexBuffer = *StaticMeshComponent->LODData[LODIndex].OverrideVertexColors;
 		FStaticMeshSourceModel& SrcModel = StaticMesh->SourceModels[LODIndex];
 		FStaticMeshRenderData& RenderData = *StaticMesh->RenderData;
 		FStaticMeshLODResources& RenderModel = RenderData.LODResources[LODIndex];
 
-		if (RenderData.WedgeMap.Num() > 0 && 
+		if (RenderData.WedgeMap.Num() > 0 &&
 			ColorVertexBuffer.GetNumVertices() == RenderModel.GetNumVertices())
 		{
 			int32 NumWedges = RawMesh.WedgeIndices.Num();
@@ -5924,7 +5938,7 @@ bool FMeshUtilities::PropagatePaintedColorsToRawMesh(UStaticMeshComponent* Stati
 				{
 					RawMesh.WedgeColors.AddUninitialized(NumWedges - NumExistingColors);
 				}
-			
+
 				for (int32 i = 0; i < NumWedges; ++i)
 				{
 					FColor WedgeColor = FColor::White;
@@ -5933,10 +5947,10 @@ bool FMeshUtilities::PropagatePaintedColorsToRawMesh(UStaticMeshComponent* Stati
 					{
 						WedgeColor = ColorVertexBuffer.VertexColor(Index);
 					}
-				
+
 					RawMesh.WedgeColors[i] = WedgeColor;
 				}
-			
+
 				return true;
 			}
 			else
@@ -5945,7 +5959,7 @@ bool FMeshUtilities::PropagatePaintedColorsToRawMesh(UStaticMeshComponent* Stati
 			}
 		}
 	}
-	
+
 	return false;
 }
 
@@ -5956,7 +5970,7 @@ static void TransformPhysicsGeometry(const FTransform& InTransform, FKAggregateG
 		FTransform ElemTM = Elem.GetTransform();
 		Elem.SetTransform(ElemTM*InTransform);
 	}
-	
+
 	for (auto& Elem : AggGeom.BoxElems)
 	{
 		FTransform ElemTM = Elem.GetTransform();
@@ -5980,12 +5994,12 @@ static void TransformPhysicsGeometry(const FTransform& InTransform, FKAggregateG
 	if (!Scale3D.Equals(FVector(1.f)))
 	{
 		const float MinPrimSize = KINDA_SMALL_NUMBER;
-		
+
 		for (auto& Elem : AggGeom.SphereElems)
 		{
 			Elem.ScaleElem(Scale3D, MinPrimSize);
 		}
-	
+
 		for (auto& Elem : AggGeom.BoxElems)
 		{
 			Elem.ScaleElem(Scale3D, MinPrimSize);
@@ -6005,16 +6019,16 @@ static void ExtractPhysicsGeometry(UStaticMeshComponent* InMeshComponent, FKAggr
 	{
 		return;
 	}
-	
+
 	if (!SrcMesh->BodySetup)
 	{
 		return;
 	}
 
 	OutAggGeom = SrcMesh->BodySetup->AggGeom;
-	
+
 	// we are not owner of this stuff
-	OutAggGeom.RenderInfo = nullptr; 
+	OutAggGeom.RenderInfo = nullptr;
 	for (auto& Elem : OutAggGeom.ConvexElems)
 	{
 		Elem.ConvexMesh = nullptr;
@@ -6030,7 +6044,7 @@ void FMeshUtilities::CalculateTextureCoordinateBoundsForRawMesh(const FRawMesh&
 {
 	const int32 NumWedges = InRawMesh.WedgeIndices.Num();
 	const int32 NumTris = NumWedges / 3;
-	
+
 	OutBounds.Empty();
 	int32 WedgeIndex = 0;
 	for (int32 TriIndex = 0; TriIndex < NumTris; TriIndex++)
@@ -6043,7 +6057,7 @@ void FMeshUtilities::CalculateTextureCoordinateBoundsForRawMesh(const FRawMesh&
 			{
 				OutBounds[MaterialIndex] += InRawMesh.WedgeTexCoords[0][WedgeIndex];
 			}
-		}		
+		}
 	}
 }
 
@@ -6078,7 +6092,7 @@ void FMeshUtilities::CalculateTextureCoordinateBoundsForSkeletalMesh(const FStat
 			FSoftSkinVertex& Vertex = Vertices[VertexIndex];
 
 			FVector2D TexCoord = Vertex.UVs[0];
-			OutBounds[MaterialIndex] += TexCoord;			
+			OutBounds[MaterialIndex] += TexCoord;
 		}
 	}
 }
@@ -6088,28 +6102,28 @@ static void CopyTextureRect(const FColor* Src, const FIntPoint& SrcSize, FColor*
 	int32 RowLength = SrcSize.X*sizeof(FColor);
 	FColor* RowDst = Dst + DstSize.X*DstPos.Y;
 	const FColor* RowSrc = Src;
-	
+
 	for (int32 RowIdx = 0; RowIdx < SrcSize.Y; ++RowIdx)
 	{
 		FMemory::Memcpy(RowDst + DstPos.X, RowSrc, RowLength);
-		
-		RowDst+= DstSize.X;
-		RowSrc+= SrcSize.X;
+
+		RowDst += DstSize.X;
+		RowSrc += SrcSize.X;
 	}
 }
 
 static void SetTextureRect(const FColor& ColorValue, const FIntPoint& SrcSize, FColor* Dst, const FIntPoint& DstSize, const FIntPoint& DstPos)
 {
 	FColor* RowDst = Dst + DstSize.X*DstPos.Y;
-	
+
 	for (int32 RowIdx = 0; RowIdx < SrcSize.Y; ++RowIdx)
 	{
 		for (int32 ColIdx = 0; ColIdx < SrcSize.X; ++ColIdx)
 		{
 			RowDst[ColIdx] = ColorValue;
 		}
-				
-		RowDst+= DstSize.X;
+
+		RowDst += DstSize.X;
 	}
 }
 
@@ -6119,7 +6133,7 @@ struct FRawMeshUVTransform
 {
 	FVector2D Offset;
 	FVector2D Scale;
-	
+
 	bool IsValid() const
 	{
 		return (Scale != FVector2D::ZeroVector);
@@ -6134,23 +6148,23 @@ static FVector2D GetValidUV(const FVector2D& UV)
 	{
 		NewUV.X = UV.X + FMath::CeilToInt(FMath::Abs(UV.X));
 	}
-	
+
 	if (UV.Y < 0.0f)
 	{
 		NewUV.Y = UV.Y + FMath::CeilToInt(FMath::Abs(UV.Y));
 	}
-	
+
 	// now make sure they're within [0, 1]
 	if (UV.X > 1.0f)
 	{
 		NewUV.X = FMath::Fmod(NewUV.X, 1.0f);
 	}
-	
+
 	if (UV.Y > 1.0f)
 	{
 		NewUV.Y = FMath::Fmod(NewUV.Y, 1.0f);
 	}
-	
+
 	return NewUV;
 }
 
@@ -6160,18 +6174,18 @@ static void MergeMaterials(UWorld* InWorld, const TArray<UMaterialInterface*>& I
 
 #if !(UE_BUILD_SHIPPING || UE_BUILD_TEST)
 	// add log to warn which material it
-	UE_LOG(LogMeshUtilities,Log,TEXT("Merging Material: Merge Material count %d"),InMaterialList.Num());
+	UE_LOG(LogMeshUtilities, Log, TEXT("Merging Material: Merge Material count %d"), InMaterialList.Num());
 
-	int32 Index=0;
+	int32 Index = 0;
 	for (auto& MaterialIter : InMaterialList)
 	{
 		if (MaterialIter)
 		{
-			UE_LOG(LogMeshUtilities,Log,TEXT("Material List %d - %s"), ++Index, *MaterialIter->GetName());
+			UE_LOG(LogMeshUtilities, Log, TEXT("Material List %d - %s"), ++Index, *MaterialIter->GetName());
 		}
 		else
 		{
-			UE_LOG(LogMeshUtilities,Log,TEXT("Material List %d - null material"), ++Index);
+			UE_LOG(LogMeshUtilities, Log, TEXT("Material List %d - null material"), ++Index);
 		}
 	}
 
@@ -6185,7 +6199,7 @@ static void MergeMaterials(UWorld* InWorld, const TArray<UMaterialInterface*>& I
 		{
 			NumOpaqueMaterials++;
 		}
-		
+
 		// Invalid UV transform
 		FRawMeshUVTransform UVTransform;
 		UVTransform.Offset = FVector2D::ZeroVector;
@@ -6198,16 +6212,16 @@ static void MergeMaterials(UWorld* InWorld, const TArray<UMaterialInterface*>& I
 		// Nothing to merge
 		return;
 	}
-		
+
 	int32 AtlasGridSize = FMath::CeilToInt(FMath::Sqrt(NumOpaqueMaterials));
 	FIntPoint AtlasTextureSize = OutMergedMaterial.DiffuseSize;
-	FIntPoint ExportTextureSize	= AtlasTextureSize/AtlasGridSize;
+	FIntPoint ExportTextureSize = AtlasTextureSize / AtlasGridSize;
 	int32 AtlasNumSamples = AtlasTextureSize.X*AtlasTextureSize.Y;
-	
-	bool bExportNormal		= (OutMergedMaterial.NormalSize != FIntPoint::ZeroValue);
-	bool bExportMetallic	= (OutMergedMaterial.MetallicSize != FIntPoint::ZeroValue);
-	bool bExportRoughness	= (OutMergedMaterial.RoughnessSize != FIntPoint::ZeroValue);
-	bool bExportSpecular	= (OutMergedMaterial.SpecularSize != FIntPoint::ZeroValue);
+
+	bool bExportNormal = (OutMergedMaterial.NormalSize != FIntPoint::ZeroValue);
+	bool bExportMetallic = (OutMergedMaterial.MetallicSize != FIntPoint::ZeroValue);
+	bool bExportRoughness = (OutMergedMaterial.RoughnessSize != FIntPoint::ZeroValue);
+	bool bExportSpecular = (OutMergedMaterial.SpecularSize != FIntPoint::ZeroValue);
 
 	// Pre-allocate buffers for texture atlas
 	OutMergedMaterial.DiffuseSamples.Reserve(AtlasNumSamples);
@@ -6239,8 +6253,8 @@ static void MergeMaterials(UWorld* InWorld, const TArray<UMaterialInterface*>& I
 
 	int32 AtlasRowIdx = 0;
 	int32 AtlasColIdx = 0;
-	FIntPoint AtlasTargetPos = FIntPoint(0,0);
-					
+	FIntPoint AtlasTargetPos = FIntPoint(0, 0);
+
 	// Flatten all materials and merge them into one material using texture atlases
 	for (int32 MatIdx = 0; MatIdx < InMaterialList.Num(); ++MatIdx)
 	{
@@ -6249,14 +6263,14 @@ static void MergeMaterials(UWorld* InWorld, const TArray<UMaterialInterface*>& I
 		{
 			continue;
 		}
-		
+
 		FFlattenMaterial FlatMaterial;
-		FlatMaterial.DiffuseSize	= ExportTextureSize;
-		FlatMaterial.NormalSize		= bExportNormal ? ExportTextureSize : FIntPoint::ZeroValue;
-		FlatMaterial.MetallicSize	= bExportMetallic ? ExportTextureSize : FIntPoint::ZeroValue;
-		FlatMaterial.RoughnessSize	= bExportRoughness ? ExportTextureSize : FIntPoint::ZeroValue;
-		FlatMaterial.SpecularSize	= bExportSpecular ? ExportTextureSize : FIntPoint::ZeroValue;
-		int32 ExportNumSamples		= ExportTextureSize.X*ExportTextureSize.Y;
+		FlatMaterial.DiffuseSize = ExportTextureSize;
+		FlatMaterial.NormalSize = bExportNormal ? ExportTextureSize : FIntPoint::ZeroValue;
+		FlatMaterial.MetallicSize = bExportMetallic ? ExportTextureSize : FIntPoint::ZeroValue;
+		FlatMaterial.RoughnessSize = bExportRoughness ? ExportTextureSize : FIntPoint::ZeroValue;
+		FlatMaterial.SpecularSize = bExportSpecular ? ExportTextureSize : FIntPoint::ZeroValue;
+		int32 ExportNumSamples = ExportTextureSize.X*ExportTextureSize.Y;
 
 		FMaterialUtilities::ExportMaterial(Material, FlatMaterial);
 
@@ -6277,7 +6291,7 @@ static void MergeMaterials(UWorld* InWorld, const TArray<UMaterialInterface*>& I
 		{
 			SetTextureRect(FlatMaterial.NormalSamples[0], ExportTextureSize, OutMergedMaterial.NormalSamples.GetData(), AtlasTextureSize, AtlasTargetPos);
 		}
-		
+
 		if (FlatMaterial.MetallicSamples.Num() == ExportNumSamples)
 		{
 			CopyTextureRect(FlatMaterial.MetallicSamples.GetData(), ExportTextureSize, OutMergedMaterial.MetallicSamples.GetData(), AtlasTextureSize, AtlasTargetPos);
@@ -6304,17 +6318,17 @@ static void MergeMaterials(UWorld* InWorld, const TArray<UMaterialInterface*>& I
 		{
 			SetTextureRect(FlatMaterial.SpecularSamples[0], ExportTextureSize, OutMergedMaterial.SpecularSamples.GetData(), AtlasTextureSize, AtlasTargetPos);
 		}
-		
+
 		check(OutUVTransforms.IsValidIndex(MatIdx));
 
 		OutUVTransforms[MatIdx].Offset = FVector2D(
-			(float)AtlasTargetPos.X/AtlasTextureSize.X, 
-			(float)AtlasTargetPos.Y/AtlasTextureSize.Y);
-		
+			(float)AtlasTargetPos.X / AtlasTextureSize.X,
+			(float)AtlasTargetPos.Y / AtlasTextureSize.Y);
+
 		OutUVTransforms[MatIdx].Scale = FVector2D(
-			(float)ExportTextureSize.X/AtlasTextureSize.X, 
-			(float)ExportTextureSize.Y/AtlasTextureSize.Y);
-		
+			(float)ExportTextureSize.X / AtlasTextureSize.X,
+			(float)ExportTextureSize.Y / AtlasTextureSize.Y);
+
 		AtlasColIdx++;
 		if (AtlasColIdx >= AtlasGridSize)
 		{
@@ -6359,7 +6373,7 @@ static void MergeFlattenedMaterials(TArray<struct FFlattenMaterial>& InMaterialL
 		UVTransform.Scale = FVector2D::ZeroVector;
 		OutUVTransforms.Add(UVTransform);
 	}
-	
+
 	int32 AtlasGridSize = FMath::CeilToInt(FMath::Sqrt(InMaterialList.Num()));
 	FIntPoint AtlasTextureSize = OutMergedMaterial.DiffuseSize;
 	FIntPoint ExportTextureSize = AtlasTextureSize / AtlasGridSize;
@@ -6545,7 +6559,7 @@ void FMeshUtilities::MergeStaticMeshComponents(const TArray<UStaticMeshComponent
 	{
 		return;
 	}
-	
+
 	SourceMeshes.AddZeroed(ComponentsToMerge.Num());
 	//SourceMeshes.SetNum(ComponentsToMerge.Num());
 
@@ -6563,7 +6577,7 @@ void FMeshUtilities::MergeStaticMeshComponents(const TArray<UStaticMeshComponent
 		UStaticMeshComponent* MeshComponent = ComponentsToMerge[MeshId];
 		// How many LOD entries merged mesh will have
 		NumMaxLOD = FMath::Max(NumMaxLOD, MeshComponent->StaticMesh->SourceModels.Num());
-				
+
 		// Save the pivot and asset package name of the first mesh, will later be used for creating merged mesh asset 
 		if (MeshId == 0)
 		{
@@ -6585,14 +6599,18 @@ void FMeshUtilities::MergeStaticMeshComponents(const TArray<UStaticMeshComponent
 	}
 
 	MainTask.EnterProgressFrame(10, LOCTEXT("MeshUtilities_MergeStaticMeshComponents_RetrievingRawMesh", "Retrieving Raw Meshes"));
-	
+
+	// Make sure if we are merging materials, that we will only export LOD0 (currently no support for baking out multiple textures and materials per LOD)
+	StartLODIndex = (InSettings.bMergeMaterials) ? 0 : StartLODIndex;
+	NumMaxLOD = (InSettings.bMergeMaterials) ? 1 : NumMaxLOD;
+
 	int32 RawMeshLODIdx = 0;
 	for (int32 LODIndex = StartLODIndex; LODIndex < NumMaxLOD; ++LODIndex, ++RawMeshLODIdx)
 	{
 		for (int32 MeshId = 0; MeshId < ComponentsToMerge.Num(); ++MeshId)
 		{
 			UStaticMeshComponent* StaticMeshComponent = ComponentsToMerge[MeshId];
-						
+
 			// We duplicate lower LOD in case this mesh has no LOD we want
 			int32 ExportLODIndex = FMath::Min(LODIndex, StaticMeshComponent->StaticMesh->SourceModels.Num() - 1);
 
@@ -6627,9 +6645,9 @@ void FMeshUtilities::MergeStaticMeshComponents(const TArray<UStaticMeshComponent
 					bOcuppiedUVChannels[RawMeshLODIdx][ChannelIdx] |= (RawMeshLOD.WedgeTexCoords[ChannelIdx].Num() != 0);
 				}
 
-				
+
 			}
-			
+
 		}
 	}
 
@@ -6640,14 +6658,14 @@ void FMeshUtilities::MergeStaticMeshComponents(const TArray<UStaticMeshComponent
 		{
 			UStaticMeshComponent* MeshComponent = ComponentsToMerge[MeshId];
 			ExtractPhysicsGeometry(MeshComponent, SourceMeshes[MeshId].AggGeom);
-			
+
 			// We will use first valid BodySetup as a source of physics settings
 			if (BodySetupSource == nullptr)
 			{
 				BodySetupSource = MeshComponent->StaticMesh->BodySetup;
 			}
 		}
-	}	
+	}
 #endif
 
 	TArray<bool> MeshShouldBakeVertexData;
@@ -6671,7 +6689,7 @@ void FMeshUtilities::MergeStaticMeshComponents(const TArray<UStaticMeshComponent
 
 	FRawMeshExt MergedMesh;
 	FMemory::Memset(&MergedMesh, 0, sizeof(MergedMesh));
-	
+
 	if (InSettings.bMergeMaterials)
 	{
 		// Should merge flattened materials into one texture
@@ -6684,7 +6702,7 @@ void FMeshUtilities::MergeStaticMeshComponents(const TArray<UStaticMeshComponent
 		FIntPoint AtlasTextureSize = InSettings.MaterialSettings.TextureSize;
 		FFlattenMaterial MergedFlatMaterial;
 		MergedFlatMaterial.DiffuseSize = AtlasTextureSize;
-		MergedFlatMaterial.NormalSize = InSettings.MaterialSettings.bNormalMap? AtlasTextureSize : FIntPoint::ZeroValue;
+		MergedFlatMaterial.NormalSize = InSettings.MaterialSettings.bNormalMap ? AtlasTextureSize : FIntPoint::ZeroValue;
 		MergedFlatMaterial.MetallicSize = InSettings.MaterialSettings.bMetallicMap ? AtlasTextureSize : FIntPoint::ZeroValue;
 		MergedFlatMaterial.RoughnessSize = InSettings.MaterialSettings.bRoughnessMap ? AtlasTextureSize : FIntPoint::ZeroValue;
 		MergedFlatMaterial.SpecularSize = InSettings.MaterialSettings.bSpecularMap ? AtlasTextureSize : FIntPoint::ZeroValue;
@@ -6731,15 +6749,15 @@ void FMeshUtilities::MergeStaticMeshComponents(const TArray<UStaticMeshComponent
 
 					// Reset material indexes
 					for (int32& FaceMaterialIndex : RawMesh.FaceMaterialIndices)
-					{						
+					{
 						FaceMaterialIndex = 0;
 					}
 				}
 				else
 				{
 					break;
-				}			
-				
+				}
+
 			}
 		}
 
@@ -6781,14 +6799,14 @@ void FMeshUtilities::MergeStaticMeshComponents(const TArray<UStaticMeshComponent
 		UniqueMaterials.Empty();
 		UniqueMaterials.Add(MergedMaterial);
 	}
-	
+
 	MainTask.EnterProgressFrame(20, LOCTEXT("MeshUtilities_MergeStaticMeshComponents_MergingMeshes", "Merging Meshes"));
 
 	// Merge meshes into single mesh
 	for (int32 SourceMeshIdx = 0; SourceMeshIdx < SourceMeshes.Num(); ++SourceMeshIdx)
 	{
 		for (int32 LODIndex = 0; LODIndex < MAX_STATIC_MESH_LODS; ++LODIndex)
-		{			
+		{
 			// Merge vertex data from source mesh list into single mesh
 			const FRawMesh& SourceRawMesh = SourceMeshes[SourceMeshIdx].MeshLODData[LODIndex].RawMesh;
 
@@ -6875,7 +6893,7 @@ void FMeshUtilities::MergeStaticMeshComponents(const TArray<UStaticMeshComponent
 					}
 				}
 			}
-			
+
 		}
 	}
 
@@ -6943,7 +6961,7 @@ void FMeshUtilities::MergeStaticMeshComponents(const TArray<UStaticMeshComponent
 		}
 
 		for (int32 LODIndex = 0; LODIndex < NumMaxLOD; ++LODIndex)
-		{			
+		{
 			FRawMesh& MergedMeshLOD = MergedMesh.MeshLODData[LODIndex].RawMesh;
 			if (MergedMeshLOD.VertexPositions.Num() > 0)
 			{
@@ -6959,7 +6977,7 @@ void FMeshUtilities::MergeStaticMeshComponents(const TArray<UStaticMeshComponent
 				SrcModel->BuildSettings.DstLightmapIndex = TargetLightMapUVChannel[LODIndex];
 
 				SrcModel->RawMeshBulkData->SaveRawMesh(MergedMeshLOD);
-			}			
+			}
 		}
 
 		// Assign materials
@@ -6979,7 +6997,7 @@ void FMeshUtilities::MergeStaticMeshComponents(const TArray<UStaticMeshComponent
 			StaticMesh->CreateBodySetup();
 			StaticMesh->BodySetup->CopyBodyPropertiesFrom(BodySetupSource);
 			StaticMesh->BodySetup->AggGeom = FKAggregateGeom();
-			
+
 			for (const FRawMeshExt& SourceMesh : SourceMeshes)
 			{
 				StaticMesh->BodySetup->AddCollisionFrom(SourceMesh.AggGeom);
@@ -6989,11 +7007,11 @@ void FMeshUtilities::MergeStaticMeshComponents(const TArray<UStaticMeshComponent
 
 		MainTask.EnterProgressFrame(10, LOCTEXT("MeshUtilities_MergeStaticMeshComponents_BuildingStaticMesh", "Building Static Mesh"));
 
-		StaticMesh->Build(bSilent);		
+		StaticMesh->Build(bSilent);
 		StaticMesh->PostEditChange();
 
 		OutAssetsToSync.Add(StaticMesh);
-		OutMergedActorLocation = MergedAssetPivot;		
+		OutMergedActorLocation = MergedAssetPivot;
 	}
 }
 
@@ -7006,7 +7024,7 @@ bool FMeshUtilities::RemoveBonesFromMesh(USkeletalMesh* SkeletalMesh, int32 LODI
 }
 
 /*------------------------------------------------------------------------------
-	Mesh reduction.
+Mesh reduction.
 ------------------------------------------------------------------------------*/
 
 IMeshReduction* FMeshUtilities::GetMeshReductionInterface()
@@ -7015,7 +7033,7 @@ IMeshReduction* FMeshUtilities::GetMeshReductionInterface()
 }
 
 /*------------------------------------------------------------------------------
-	Mesh merging.
+Mesh merging.
 ------------------------------------------------------------------------------*/
 IMeshMerging* FMeshUtilities::GetMeshMergingInterface()
 {
@@ -7023,7 +7041,7 @@ IMeshMerging* FMeshUtilities::GetMeshMergingInterface()
 }
 
 /*------------------------------------------------------------------------------
-	Module initialization / teardown.
+Module initialization / teardown.
 ------------------------------------------------------------------------------*/
 
 void FMeshUtilities::StartupModule()
@@ -7040,18 +7058,18 @@ void FMeshUtilities::StartupModule()
 		TArray<FName> SwarmModuleNames;
 		FModuleManager::Get().FindModules(TEXT("*SimplygonSwarm"), SwarmModuleNames);
 
-		
+
 		for (int32 Index = 0; Index < ModuleNames.Num(); Index++)
 		{
 			IMeshReductionModule& MeshReductionModule = FModuleManager::LoadModuleChecked<IMeshReductionModule>(ModuleNames[Index]);
-								
+
 			// Look for MeshReduction interface
 			if (MeshReduction == NULL)
 			{
 				MeshReduction = MeshReductionModule.GetMeshReductionInterface();
 				if (MeshReduction)
 				{
-					UE_LOG(LogMeshUtilities,Log,TEXT("Using %s for automatic mesh reduction"),*ModuleNames[Index].ToString());
+					UE_LOG(LogMeshUtilities, Log, TEXT("Using %s for automatic mesh reduction"), *ModuleNames[Index].ToString());
 				}
 			}
 
@@ -7061,7 +7079,7 @@ void FMeshUtilities::StartupModule()
 				MeshMerging = MeshReductionModule.GetMeshMergingInterface();
 				if (MeshMerging)
 				{
-					UE_LOG(LogMeshUtilities,Log,TEXT("Using %s for automatic mesh merging"),*ModuleNames[Index].ToString());
+					UE_LOG(LogMeshUtilities, Log, TEXT("Using %s for automatic mesh merging"), *ModuleNames[Index].ToString());
 				}
 			}
 
@@ -7071,7 +7089,7 @@ void FMeshUtilities::StartupModule()
 				break;
 			}
 		}
-		
+
 
 		for (int32 Index = 0; Index < SwarmModuleNames.Num(); Index++)
 		{
@@ -7091,12 +7109,12 @@ void FMeshUtilities::StartupModule()
 
 		if (!MeshReduction)
 		{
-			UE_LOG(LogMeshUtilities,Log,TEXT("No automatic mesh reduction module available"));
+			UE_LOG(LogMeshUtilities, Log, TEXT("No automatic mesh reduction module available"));
 		}
 
 		if (!MeshMerging)
 		{
-			UE_LOG(LogMeshUtilities,Log,TEXT("No automatic mesh merging module available"));
+			UE_LOG(LogMeshUtilities, Log, TEXT("No automatic mesh merging module available"));
 		}
 		else
 		{
