Copying //UE4/Dev-Framework to Dev-Main (//UE4/Dev-Main) @ 2855699

==========================
MAJOR FEATURES + CHANGES
==========================

Change 2839897 on 2016/01/22 by [REDACTED_PERSON]

	Allow static mesh editor to specify a default collision profile.

	#UE-2836

Change 2840489 on 2016/01/22 by [REDACTED_PERSON]

	Fix collision customization so that it respects const editing property

Change 2840528 on 2016/01/22 by [REDACTED_PERSON]

	Fix compile error and actually get value from attribute

Change 2840672 on 2016/01/22 by [REDACTED_PERSON]

	#ue4 - Include data from USkinnedMeshComponent in USkeletalMeshComponent::GetResourceSize().

Change 2841314 on 2016/01/24 by [REDACTED_PERSON]

	Fix depressingly frequent misspellings of 'suppress'

Change 2841323 on 2016/01/24 by [REDACTED_PERSON]

	Reserve worst case memory for TSet Intersect, Union, and Difference to avoid memory allocations during iteration
	Ensure that TSet Intersect considers the least number of elements possible
	Early out from TSet Contains if Other is larger than this
	Clarify comment on TSet Difference

Change 2841380 on 2016/01/24 by [REDACTED_PERSON]

	UE-25586 Audio assets not correctly reporting resource memory usage

	Tested on PC/PS4 and with Editor builds. Memory reporting is working for all cases now.

Change 2841385 on 2016/01/24 by [REDACTED_PERSON]

	UE-21210 Adding subtitle priority to USoundWave

Change 2841386 on 2016/01/24 by [REDACTED_PERSON]

	Return null for GameNetDriver if World is null instead of crashing

Change 2841409 on 2016/01/24 by [REDACTED_PERSON]

	UE-25514 Removing load for default objects for every sound wave

Change 2841858 on 2016/01/25 by [REDACTED_PERSON]

	Make sure that PIE face index results are consistent with runtime

Change 2841977 on 2016/01/25 by [REDACTED_PERSON]

	Fix object type customization so that it's only enabled when custom is selected. (Accidently broke this in recent change)

Change 2841982 on 2016/01/25 by [REDACTED_PERSON]

	Minor optimization by avoiding recreating FNames repeatedly in constructor

Change 2842169 on 2016/01/25 by [REDACTED_PERSON]

	Fixes to animBP compiler and instance to store and double buffer internal machine state weights on the instance. So they can be queried cross-machine without issue.

Change 2842390 on 2016/01/25 by [REDACTED_PERSON]

	Fix in world editing of BodyInstance not working.

	No longer serializing Scale3D as this is allways initialized in InitBody.
	No longer overwriting MassInKg and renamed to to MassInKgOverride which better reflects what this variable does.

Change 2843579 on 2016/01/26 by [REDACTED_PERSON]

	Only update replication when it actually changes
	Don't check calling SetIsReplicated if the class cannot replicate, instead output an error message
	Fix spelling in comment

Change 2843627 on 2016/01/26 by [REDACTED_PERSON]

	Add \\ as a default console key for Italian keyboard layouts

Change 2843628 on 2016/01/26 by [REDACTED_PERSON]

	Don't reconstruct FName on each call to GetHitResultAtScreenPosition

Change 2843671 on 2016/01/26 by [REDACTED_PERSON]

	Fix incorrect bone transforms being pushed to the renderer during SetSkeletalMesh. This presented as motion blur artifacts in editor

Change 2843768 on 2016/01/26 by [REDACTED_PERSON]

	Inline Get Component functions in TriggerBase

Change 2844003 on 2016/01/26 by [REDACTED_PERSON]

	#ue4 - Fix FMath::Fmod(X, Y) sometimes returning small negative values for positive X and Y due to float imprecision. Added tests to math tests at startup to check this, and also to better handle results close to Y. Wrap the ensure on Y=0 within a conditional so a breakpoint can be used during debugging (to distinguish between zero and very small input).

Change 2844005 on 2016/01/26 by [REDACTED_PERSON]

	#ue4 - Convert uses of fmod() and fmodf() to use FMath::Fmod() instead.

	Also see CL 2844003

diff --git a/Engine/Config/BaseEngine.ini b/Engine/Config/BaseEngine.ini
--- a/Engine/Config/BaseEngine.ini
+++ b/Engine/Config/BaseEngine.ini
@@ -144,6 +144,13 @@ LightingOnlyBrightness=(R=0.3,G=0.3,B=0.3,A=1.0)
 +StreamingAccuracyColors=(R=0.7,G=0.7,B=0.7,A=1.0)
 +StreamingAccuracyColors=(R=0.0,G=0.8,B=0.5,A=1.0)
 +StreamingAccuracyColors=(R=0.0,G=1.0,B=0.0,A=1.0)
++HLODColorationColors=(R=1.0,G=1.0,B=1.0,A=1.0)
++HLODColorationColors=(R=0.0,G=1.0,B=0.0,A=1.0)
++HLODColorationColors=(R=0.0,G=0.0,B=1.0,A=1.0)
++HLODColorationColors=(R=1.0,G=1.0,B=0.0,A=1.0)
++HLODColorationColors=(R=1.0,G=0.0,B=1.0,A=1.0)
++HLODColorationColors=(R=0.0,G=1.0,B=1.0,A=1.0)
++HLODColorationColors=(R=0.5,G=0.5,B=0.5,A=1.0)
 MaxPixelShaderAdditiveComplexityCount=2000
 MaxES2PixelShaderAdditiveComplexityCount=600
 bSubtitlesEnabled=True
@@ -290,6 +297,9 @@ bUseStreamingPause=false
 +K2FieldRedirects=(OldFieldName="SingleAnimSkeletalComponent.GetPlayRate",NewFieldName="SkeletalMeshComponent.GetPlayRate")
 +K2FieldRedirects=(OldFieldName="SingleAnimSkeletalComponent.IsPlaying",NewFieldName="SkeletalMeshComponent.IsPlaying")
 
++TaggedPropertyRedirects=(ClassName="BodyInstance",OldPropertyName="MassInKg",NewPropertyName="MassInKgOverride")
+
+
 +K2FieldRedirects=(OldFieldName="SkeletalMeshComponent.SetAnimBlueprint",NewFieldName="SkeletalMeshComponent.SetAnimInstanceClass")
 +K2ParamRedirects=(NodeName="/Script/BlueprintGraph.K2Node_CallFunction", OldParamName="SkeletalMeshComponent.SetAnimInstanceClass.NewBlueprint", NewParamName="NewClass")
 +K2FieldRedirects=(OldFieldName="SkeletalMeshComponent.SetAnimClass",NewFieldName="SkeletalMeshComponent.SetAnimInstanceClass")
@@ -1019,6 +1029,8 @@ bUseStreamingPause=false
 +K2ParamRedirects=(NodeName="/Script/BlueprintGraph.K2Node_CallFunction", OldParamName="KismetSystemLibrary.ComponentOverlapActors_NEW.ObjectsToTrace", NewParamName="ObjectTypes")
 +K2ParamRedirects=(NodeName="/Script/BlueprintGraph.K2Node_CallFunction", OldParamName="KismetSystemLibrary.ComponentOverlapComponents_NEW.ObjectsToTrace", NewParamName="ObjectTypes")
 
++K2FieldRedirects=(OldFieldName="KismetSystemLibrary.SetSupressViewportTransitionMessage",NewFieldName="KismetSystemLibrary.SetSuppressViewportTransitionMessage")
+ 
 +K2FieldRedirects=(OldFieldName="SceneComponent.SetRelativeTranslation",NewFieldName="SceneComponent.K2_SetRelativeLocation")
 +K2ParamRedirects=(NodeName="/Script/BlueprintGraph.K2Node_CallFunction", OldParamName="SceneComponent.K2_SetRelativeLocation.NewTranslation", NewParamName="NewLocation")
 +K2FieldRedirects=(OldFieldName="SceneComponent.AddLocalTranslation",NewFieldName="SceneComponent.K2_AddLocalOffset")
diff --git a/Engine/Plugins/Runtime/LeapMotionController/Source/LeapMotionController/Private/LeapMotionHandActor.cpp b/Engine/Plugins/Runtime/LeapMotionController/Source/LeapMotionController/Private/LeapMotionHandActor.cpp
--- a/Engine/Plugins/Runtime/LeapMotionController/Source/LeapMotionController/Private/LeapMotionHandActor.cpp
+++ b/Engine/Plugins/Runtime/LeapMotionController/Source/LeapMotionController/Private/LeapMotionHandActor.cpp
@@ -22,7 +22,7 @@ ALeapMotionHandActor::ALeapMotionHandActor(const class FObjectInitializer& Objec
 
 ULeapMotionControllerComponent* ALeapMotionHandActor::GetParentingControllerComponent() const
 {
-	return Cast<ULeapMotionControllerComponent>(RootComponent->AttachParent);
+	return Cast<ULeapMotionControllerComponent>(RootComponent->GetAttachParent());
 }
 
 ALeapMotionBoneActor* ALeapMotionHandActor::GetBoneActor(ELeapBone LeapBone) const
diff --git a/Engine/Plugins/Runtime/ProceduralMeshComponent/Source/Private/ProceduralMeshComponent.cpp b/Engine/Plugins/Runtime/ProceduralMeshComponent/Source/Private/ProceduralMeshComponent.cpp
--- a/Engine/Plugins/Runtime/ProceduralMeshComponent/Source/Private/ProceduralMeshComponent.cpp
+++ b/Engine/Plugins/Runtime/ProceduralMeshComponent/Source/Private/ProceduralMeshComponent.cpp
@@ -394,7 +394,10 @@ void UProceduralMeshComponent::CreateMeshSection(int32 SectionIndex, const TArra
 	SCOPE_CYCLE_COUNTER(STAT_ProcMesh_CreateMeshSection);
 
 	// Ensure sections array is long enough
-	ProcMeshSections.SetNum(SectionIndex + 1, false);
+	if (SectionIndex >= ProcMeshSections.Num())
+	{
+		ProcMeshSections.SetNum(SectionIndex + 1, false);
+	}
 
 	// Reset this section (in case it already existed)
 	FProcMeshSection& NewSection = ProcMeshSections[SectionIndex];
diff --git a/Engine/Shaders/GpuSkinVertexFactory.usf b/Engine/Shaders/GpuSkinVertexFactory.usf
--- a/Engine/Shaders/GpuSkinVertexFactory.usf
+++ b/Engine/Shaders/GpuSkinVertexFactory.usf
@@ -89,7 +89,7 @@ struct FVertexFactoryInput
 	/** added to the Position */
 	float3	DeltaPosition	: ATTRIBUTE9;	//POSITION1;
 	/** added to the TangentZ and then used to derive new TangentX,TangentY */
-	half4	DeltaTangentZ	: ATTRIBUTE10;	//NORMAL1;
+	float3	DeltaTangentZ	: ATTRIBUTE10;	//NORMAL1;
 #elif GPUSKIN_APEX_CLOTH // exclusive with GPUSKIN_MORPH_BLEND
 	// APEX cloth mesh-mesh mapping data
 	// Barycentric Coordinate Data
@@ -378,13 +378,9 @@ float3x3 SkinTangents( FVertexFactoryInput Input, FVertexFactoryIntermediates In
 	float3 TangentX = TangentBias(Input.TangentX);
 	// normal
 	float4 TangentZ = TangentBias(Input.TangentZ);
-	// normal morph offset. 
-	// since normal delta can be in the range of (-2, 2), 
-	// we packed twice more with less precision, so scale back up by 2 
-	float3 DeltaTangentZ = TangentBias(Input.DeltaTangentZ.xyz)*2;
 
 	// calc new normal by offseting it with the delta
-	Tangents[1] = normalize( TangentZ.xyz + Input.DeltaTangentZ.w * DeltaTangentZ);
+	Tangents[1] = normalize( TangentZ.xyz + Input.DeltaTangentZ.xyz );
 	// derive the new tangent by orthonormalizing the new normal against
 	// the base tangent vector (assuming these are normalized)
 	Tangents[0] = normalize( TangentX - (dot(TangentX,Tangents[1]) * Tangents[1]) );
diff --git a/Engine/Shaders/PostProcessVisualizeHDR.usf b/Engine/Shaders/PostProcessVisualizeHDR.usf
--- a/Engine/Shaders/PostProcessVisualizeHDR.usf
+++ b/Engine/Shaders/PostProcessVisualizeHDR.usf
@@ -40,7 +40,7 @@ uint ComputeAdvice(float3 HDRColor)
 		return 1;
 	}
 
-	// for some reasons HLSL compiler seems to supress  the return 1 when we comment in this code
+	// for some reasons HLSL compiler seems to suppress  the return 1 when we comment in this code
 //	if(Lum > EyeAdaptationParams[0].w)
 //	{
 //		return 2;
diff --git a/Engine/Shaders/SimpleElementPixelShader.usf b/Engine/Shaders/SimpleElementPixelShader.usf
--- a/Engine/Shaders/SimpleElementPixelShader.usf
+++ b/Engine/Shaders/SimpleElementPixelShader.usf
@@ -166,7 +166,7 @@ half Gamma;
 void GammaMain(
 	in float2 TextureCoordinate : TEXCOORD0,
 	in float4 Color : TEXCOORD1,
-	in float4 InHitProxyId : TEXCOORD2,			// needed to have d3ddebug warning supressed (VS and PS signatures with SV_Position
+	in float4 InHitProxyId : TEXCOORD2,			// needed to have d3ddebug warning suppressed (VS and PS signatures with SV_Position
 	out float4 OutColor : SV_Target0
 #if ENABLE_EDITOR_COMPOSITING
 	,in float4 SvPosition : SV_Position
@@ -201,7 +201,7 @@ void GammaMain(
 void GammaAlphaOnlyMain(
 	in float2 TextureCoordinate : TEXCOORD0,
 	in float4 Color : TEXCOORD1,
-	in float4 InHitProxyId : TEXCOORD2,			// needed to have d3ddebug warning supressed (VS and PS signatures with SV_Position
+	in float4 InHitProxyId : TEXCOORD2,			// needed to have d3ddebug warning suppressed (VS and PS signatures with SV_Position
 	out float4 OutColor : SV_Target0
 #if ENABLE_EDITOR_COMPOSITING
 	,in float4 SvPosition : SV_Position
diff --git a/Engine/Source/Developer/BlueprintCompilerCppBackend/Private/BlueprintCompilerCppBackendValueHelper.cpp b/Engine/Source/Developer/BlueprintCompilerCppBackend/Private/BlueprintCompilerCppBackendValueHelper.cpp
--- a/Engine/Source/Developer/BlueprintCompilerCppBackend/Private/BlueprintCompilerCppBackendValueHelper.cpp
+++ b/Engine/Source/Developer/BlueprintCompilerCppBackend/Private/BlueprintCompilerCppBackendValueHelper.cpp
@@ -681,7 +681,7 @@ void FEmitDefaultValueHelper::GenerateConstructor(FEmitterLocalContext& Context)
 					if (NativeRootComponentFallback.IsEmpty())
 					{
 						USceneComponent* SceneComponent = Cast<USceneComponent>(DSO);
-						if (SceneComponent && !SceneComponent->AttachParent && SceneComponent->CreationMethod == EComponentCreationMethod::Native)
+						if (SceneComponent && !SceneComponent->GetAttachParent() && SceneComponent->CreationMethod == EComponentCreationMethod::Native)
 						{
 							NativeRootComponentFallback = VariableName;
 						}
diff --git a/Engine/Source/Developer/HierarchicalLODUtilities/Private/HierarchicalLODUtilities.cpp b/Engine/Source/Developer/HierarchicalLODUtilities/Private/HierarchicalLODUtilities.cpp
--- a/Engine/Source/Developer/HierarchicalLODUtilities/Private/HierarchicalLODUtilities.cpp
+++ b/Engine/Source/Developer/HierarchicalLODUtilities/Private/HierarchicalLODUtilities.cpp
@@ -24,6 +24,7 @@
 #include "Factories/Factory.h"
 #include "ObjectTools.h"
 #include "AssetEditorManager.h"
+#include "ScopedTransaction.h"
 #endif // WITH_EDITOR
 
 #include "HierarchicalLODUtilities.h"
@@ -111,21 +112,13 @@ bool FHierarchicalLODUtilities::BuildStaticMeshForLODActor(ALODActor* LODActor,
 			return false;
 		}
 
-		// Delete actor assets before generating new ones
-		FHierarchicalLODUtilities::DeleteLODActorAssets(LODActor);	
-
-		// Clean up sub objects (if applicable)
-		for (auto& SubOject : LODActor->SubObjects)
-		{
-			if (SubOject)
-			{
-				SubOject->MarkPendingKill();
-			}
-		}
-		LODActor->SubObjects.Empty();
+		const FScopedTransaction Transaction(LOCTEXT("UndoAction_BuildProxyMesh", "Building Proxy Mesh for Cluster"));
+		LODActor->Modify();		
 
+		// Delete actor assets before generating new ones
+		FHierarchicalLODUtilities::DestroyClusterData(LODActor);	
+		
 		TArray<UStaticMeshComponent*> AllComponents;
-
 		{
 			FScopedSlowTask SlowTask(LODActor->SubActors.Num(), (LOCTEXT("HierarchicalLODUtils_CollectStaticMeshes", "Collecting Static Meshes for Cluster")));
 			SlowTask.MakeDialog();
@@ -195,7 +188,7 @@ bool FHierarchicalLODUtilities::BuildStaticMeshForLODActor(ALODActor* LODActor,
 				}
 
 				FGuid JobID = Processor->AddProxyJob(LODActor, OverrideLODSetup);
-				MeshUtilities.CreateProxyMesh(Actors, ProxySettings, AssetsOuter, PackageName, JobID, Processor->GetCallbackDelegate(), true);
+				MeshUtilities.CreateProxyMesh(Actors, ProxySettings, AssetsOuter, PackageName, JobID, Processor->GetCallbackDelegate(), true, OverrideLODSetup.TransitionScreenSize);
 				return true;
 			}
 			else
@@ -235,7 +228,7 @@ bool FHierarchicalLODUtilities::BuildStaticMeshForLODActor(ALODActor* LODActor,
 				FBoxSphereBounds Bounds = LODActor->GetStaticMeshComponent()->CalcBounds(FTransform());
 				LODActor->LODDrawDistance = FHierarchicalLODUtilities::CalculateDrawDistanceFromScreenSize(Bounds.SphereRadius, LODSetup.TransitionScreenSize, ProjectionMatrix);
 				LODActor->UpdateSubActorLODParents();
-				
+
 				// Freshly build so mark not dirty
 				LODActor->SetIsDirty(false);
 
@@ -331,50 +324,45 @@ ALODActor* FHierarchicalLODUtilities::GetParentLODActor(const AActor* InActor)
 	return ParentActor;
 }
 
-void FHierarchicalLODUtilities::DeleteLODActor(ALODActor* InActor)
+void FHierarchicalLODUtilities::DestroyCluster(ALODActor* InActor)
 {
 	// Find if it has a parent ALODActor
 	AActor* Actor = CastChecked<AActor>(InActor);
+	UWorld* World = Actor->GetWorld();
+	ALODActor* ParentLOD = GetParentLODActor(InActor);
+
+	const FScopedTransaction Transaction(LOCTEXT("UndoAction_DeleteCluster", "Deleting a (invalid) Cluster"));
+	Actor->Modify();
+	World->Modify();
 
-	ALODActor* ParentLOD = GetParentLODActor(Actor);
-	if (ParentLOD)
+	if (ParentLOD != nullptr)
 	{
+		ParentLOD->Modify();
 		ParentLOD->RemoveSubActor(Actor);
 	}
 
 	// Clean out sub actors and update their LODParent
 	while (InActor->SubActors.Num())
 	{
-		auto SubActor = InActor->SubActors[0];
+		AActor* SubActor = InActor->SubActors[0];
+		SubActor->Modify();
 		InActor->RemoveSubActor(SubActor);
 	}
 
-	// you still have to delete all objects just in case they had it and didn't want it anymore
-	TArray<UObject*> AssetsToDelete;
+	// Also destroy the cluster's data
+	DestroyClusterData(InActor);
 
-	AssetsToDelete.Append(InActor->SubObjects);
+	World->DestroyActor(InActor);
 
-	for (auto& Asset : AssetsToDelete)
+	if (ParentLOD != nullptr && !ParentLOD->HasValidSubActors())
 	{
-		if (Asset)
-		{
-			Asset->MarkPendingKill();
-			ObjectTools::DeleteSingleObject(Asset, false);
-		}
+		DestroyCluster(ParentLOD);
 	}
-
-	// Rebuild streaming data to prevent invalid references
-	ULevel::BuildStreamingData(InActor->GetWorld(), InActor->GetLevel());
-
-	// garbage collect
-	CollectGarbage(GARBAGE_COLLECTION_KEEPFLAGS, true);
 }
 
-void FHierarchicalLODUtilities::DeleteLODActorAssets(ALODActor* InActor)
+void FHierarchicalLODUtilities::DestroyClusterData(ALODActor* InActor)
 {
-	// you still have to delete all objects just in case they had it and didn't want it anymore
 	TArray<UObject*> AssetsToDelete;
-
 	AssetsToDelete.Append(InActor->SubObjects);
 	InActor->SubObjects.Empty();
 
@@ -386,21 +374,12 @@ void FHierarchicalLODUtilities::DeleteLODActorAssets(ALODActor* InActor)
 			// Close possible open editors using this asset	
 			FAssetEditorManager::Get().CloseAllEditorsForAsset(AssetObject);
 #endif // WITH_EDITOR
-			AssetObject->MarkPendingKill();
-			ObjectTools::DeleteSingleObject(AssetObject, false);
+			AssetObject->ClearFlags(RF_AllFlags);
 		}
 	}
 
-	InActor->GetStaticMeshComponent()->StaticMesh = nullptr;
-
-	// Mark render state dirty
-	InActor->GetStaticMeshComponent()->MarkRenderStateDirty();
-
-	// Rebuild streaming data to prevent invalid references
-	ULevel::BuildStreamingData(InActor->GetWorld(), InActor->GetLevel());
-
-	// Garbage collecting
-	CollectGarbage(GARBAGE_COLLECTION_KEEPFLAGS, true);
+	// Set Static Mesh to null since there isn't a mesh anymore
+	InActor->GetStaticMeshComponent()->SetStaticMesh(nullptr);
 }
 
 ALODActor* FHierarchicalLODUtilities::CreateNewClusterActor(UWorld* InWorld, const int32 InLODLevel, AWorldSettings* WorldSettings)
@@ -422,6 +401,270 @@ ALODActor* FHierarchicalLODUtilities::CreateNewClusterActor(UWorld* InWorld, con
 	return NewActor;
 }
 
+ALODActor* FHierarchicalLODUtilities::CreateNewClusterFromActors(UWorld* InWorld, AWorldSettings* WorldSettings, const TArray<AActor*>& InActors, const int32 InLODLevel /*= 0*/)
+{
+	checkf(InWorld != nullptr, TEXT("Invalid world"));
+	checkf(InActors.Num() > 0, TEXT("Zero number of sub actors"));
+	checkf(WorldSettings != nullptr, TEXT("Invalid world settings"));
+	checkf(WorldSettings->bEnableHierarchicalLODSystem, TEXT("Hierarchical LOD system is disabled"));
+
+	const FScopedTransaction Transaction(LOCTEXT("UndoAction_CreateNewCluster", "Create new Cluster"));
+	InWorld->Modify();
+
+	// Create the cluster
+	ALODActor* NewCluster = FHierarchicalLODUtilities::CreateNewClusterActor(InWorld, InLODLevel, WorldSettings);
+	checkf(NewCluster != nullptr, TEXT("Failed to create a new cluster"));
+
+	// Add InActors to new cluster
+	for (AActor* Actor : InActors)
+	{
+		checkf(Actor != nullptr, TEXT("Invalid actor in InActors"));
+		
+		// Check if Actor is currently part of a different cluster
+		ALODActor* ParentActor = FHierarchicalLODUtilities::GetParentLODActor(Actor);
+		if (ParentActor != nullptr)
+		{
+			// If so remove it first
+			ParentActor->Modify();
+			ParentActor->RemoveSubActor(Actor);
+
+			// If the parent cluster is now empty (invalid) destroy it
+			if (!ParentActor->HasValidSubActors())
+			{
+				FHierarchicalLODUtilities::DestroyCluster(ParentActor);
+			}
+		}
+
+		// Add actor to new cluster
+		NewCluster->AddSubActor(Actor);
+	}
+
+	// Update sub actor LOD parents to populate 
+	NewCluster->UpdateSubActorLODParents();
+
+	return NewCluster;
+}
+
+const bool FHierarchicalLODUtilities::RemoveActorFromCluster(AActor* InActor)
+{
+	checkf(InActor != nullptr, TEXT("Invalid InActor"));
+	
+	bool bSucces = false;
+
+	ALODActor* ParentActor = FHierarchicalLODUtilities::GetParentLODActor(InActor);
+	if (ParentActor != nullptr)
+	{
+		const FScopedTransaction Transaction(LOCTEXT("UndoAction_RemoveActorFromCluster", "Remove Actor From Cluster"));
+		ParentActor->Modify();
+		InActor->Modify();
+
+		bSucces = ParentActor->RemoveSubActor(InActor);
+
+		if (!ParentActor->HasValidSubActors())
+		{
+			DestroyCluster(ParentActor);
+		}
+	}
+	
+	return bSucces;
+}
+
+const bool FHierarchicalLODUtilities::AddActorToCluster(AActor* InActor, ALODActor* InParentActor)
+{
+	checkf(InActor != nullptr, TEXT("Invalid InActor"));
+	checkf(InParentActor != nullptr, TEXT("Invalid InParentActor"));
+
+	// First, if it is the case remove the actor from it's current cluster
+	const bool bActorWasClustered = RemoveActorFromCluster(InActor);
+
+	// Now add it to the new one
+	const FScopedTransaction Transaction(LOCTEXT("UndoAction_AddActorToCluster", "Add Actor To Cluster"));
+	InParentActor->Modify();
+	InActor->Modify();
+
+	// Add InActor to InParentActor cluster
+	InParentActor->AddSubActor(InActor);
+
+#if WITH_EDITOR
+	if (bActorWasClustered)
+	{
+		GEditor->BroadcastHLODActorAdded(InActor, InParentActor);
+	}
+	else
+	{
+		GEditor->BroadcastHLODActorAdded(InActor, InParentActor);
+	}
+
+#endif // WITH_EDITOR
+
+	return true;
+}
+
+const bool FHierarchicalLODUtilities::MergeClusters(ALODActor* TargetCluster, ALODActor* SourceCluster)
+{
+	checkf(TargetCluster != nullptr&& TargetCluster->SubActors.Num() > 0, TEXT("Invalid InActor"));
+	checkf(SourceCluster != nullptr && SourceCluster->SubActors.Num() > 0, TEXT("Invalid InParentActor"));
+
+	const FScopedTransaction Transaction(LOCTEXT("UndoAction_MergeClusters", "Merge Clusters"));
+	TargetCluster->Modify();
+	SourceCluster->Modify();
+
+	while (SourceCluster->SubActors.Num())
+	{
+		AActor* SubActor = SourceCluster->SubActors.Last();
+		FHierarchicalLODUtilities::AddActorToCluster(SubActor, TargetCluster);		
+	}
+
+	if (!SourceCluster->HasValidSubActors())
+	{
+		DestroyCluster(SourceCluster);
+	}
+
+	return true;
+}
+
+const bool FHierarchicalLODUtilities::AreActorsInSamePersistingLevel(const TArray<AActor*>& InActors)
+{
+	ULevel* Level = nullptr;
+	bool bSameLevelInstance = true;
+	for (AActor* Actor : InActors)
+	{
+		if (Level == nullptr)
+		{
+			Level = Actor->GetLevel();
+		}
+
+		bSameLevelInstance &= (Level == Actor->GetLevel());
+
+		if (!bSameLevelInstance)
+		{
+			break;
+		}
+	}
+
+	return bSameLevelInstance;
+}
+
+const bool FHierarchicalLODUtilities::AreClustersInSameHLODLevel(const TArray<ALODActor*>& InLODActors)
+{
+	int32 HLODLevel = -1;
+	bool bSameHLODLevel = true;
+	for (ALODActor* LODActor : InLODActors)
+	{
+		if (HLODLevel == -1)
+		{
+			HLODLevel = LODActor->LODLevel;
+		}
+
+		bSameHLODLevel &= (HLODLevel == LODActor->LODLevel);
+
+		if (!bSameHLODLevel)
+		{
+			break;
+		}
+	}
+
+	return bSameHLODLevel;
+}
+
+const bool FHierarchicalLODUtilities::AreActorsInSameHLODLevel(const TArray<AActor*>& InActors)
+{
+	int32 HLODLevel = -1;
+	bool bSameHLODLevel = true;
+	for (AActor* Actor : InActors)
+	{
+		ALODActor* ParentActor = FHierarchicalLODUtilities::GetParentLODActor(Actor);
+
+		if (ParentActor != nullptr)
+		{
+			if (HLODLevel == -1)
+			{
+				HLODLevel = ParentActor->LODLevel;
+			}
+
+			bSameHLODLevel &= (HLODLevel == ParentActor->LODLevel);
+		}
+		else
+		{
+			bSameHLODLevel = false;
+		}
+
+		if (!bSameHLODLevel)
+		{
+			break;
+		}
+	}
+
+	return bSameHLODLevel;
+}
+
+const bool FHierarchicalLODUtilities::AreActorsClustered(const TArray<AActor*>& InActors)
+{	
+	bool bClustered = true;
+	for (AActor* Actor : InActors)
+	{
+		bClustered &= (FHierarchicalLODUtilities::GetParentLODActor(Actor) != nullptr);
+
+		if (!bClustered)
+		{
+			break;
+		}
+	}
+
+	return bClustered;
+}
+
+const bool FHierarchicalLODUtilities::IsActorClustered(const AActor* InActor)
+{
+	bool bClustered = (FHierarchicalLODUtilities::GetParentLODActor(InActor) != nullptr);	
+	return bClustered;
+}
+
+void FHierarchicalLODUtilities::ExcludeActorFromClusterGeneration(AActor* InActor)
+{
+	const FScopedTransaction Transaction(LOCTEXT("UndoAction_ExcludeActorFromClusterGeneration", "Exclude Actor From Cluster Generation"));
+	InActor->Modify();
+	InActor->bEnableAutoLODGeneration = false;
+	FHierarchicalLODUtilities::RemoveActorFromCluster(InActor);
+}
+
+void FHierarchicalLODUtilities::DestroyLODActor(ALODActor* InActor)
+{
+	const FScopedTransaction Transaction(LOCTEXT("UndoAction_DeleteLODActor", "Delete LOD Actor"));
+	UWorld* World = InActor->GetWorld();
+	World->Modify();
+	InActor->Modify();
+	
+	ALODActor* ParentActor = FHierarchicalLODUtilities::GetParentLODActor(InActor);
+
+	FHierarchicalLODUtilities::DestroyCluster(InActor);
+	World->DestroyActor(InActor);
+
+	if (ParentActor && !ParentActor->HasValidSubActors())
+	{
+		ParentActor->Modify();
+		FHierarchicalLODUtilities::DestroyLODActor(ParentActor);
+	}
+}
+
+void FHierarchicalLODUtilities::ExtractStaticMeshActorsFromLODActor(ALODActor* LODActor, TArray<AActor*> &InOutActors)
+{
+	for (auto ChildActor : LODActor->SubActors)
+	{
+		if (ChildActor)
+		{
+			TArray<AActor*> ChildActors;
+			if (ChildActor->IsA<ALODActor>())
+			{
+				FHierarchicalLODUtilities::ExtractStaticMeshActorsFromLODActor(Cast<ALODActor>(ChildActor), ChildActors);
+			}
+
+			ChildActors.Push(ChildActor);
+			InOutActors.Append(ChildActors);
+		}
+	}	
+}
+
 void FHierarchicalLODUtilities::DeleteLODActorsInHLODLevel(UWorld* InWorld, const int32 HLODLevelIndex)
 {
 	// you still have to delete all objects just in case they had it and didn't want it anymore
@@ -431,19 +674,19 @@ void FHierarchicalLODUtilities::DeleteLODActorsInHLODLevel(UWorld* InWorld, cons
 		ALODActor* LodActor = Cast<ALODActor>(InWorld->PersistentLevel->Actors[ActorId]);
 		if (LodActor && LodActor->LODLevel == (HLODLevelIndex + 1))
 		{
-			DeleteLODActor(LodActor);
+			DestroyCluster(LodActor);
 			InWorld->DestroyActor(LodActor);
 		}
 	}
 }
 
-int32 FHierarchicalLODUtilities::ComputeStaticMeshLODLevel(const TArray<FStaticMeshSourceModel>& SourceModels, const float ScreenAreaSize)
-{
+int32 FHierarchicalLODUtilities::ComputeStaticMeshLODLevel(const TArray<FStaticMeshSourceModel>& SourceModels, const FStaticMeshRenderData* RenderData, const float ScreenAreaSize)
+{	
 	const int32 NumLODs = SourceModels.Num();
 	// Walk backwards and return the first matching LOD
 	for (int32 LODIndex = NumLODs - 1; LODIndex >= 0; --LODIndex)
 	{
-		if (SourceModels[LODIndex].ScreenSize > ScreenAreaSize)
+		if (SourceModels[LODIndex].ScreenSize > ScreenAreaSize || ((SourceModels[LODIndex].ScreenSize == 0.0f) && (RenderData->ScreenSize[LODIndex] != SourceModels[LODIndex].ScreenSize) && (RenderData->ScreenSize[LODIndex] > ScreenAreaSize)))
 		{
 			return FMath::Max(LODIndex, 0);
 		}
@@ -452,9 +695,14 @@ int32 FHierarchicalLODUtilities::ComputeStaticMeshLODLevel(const TArray<FStaticM
 	return 0;
 }
 
-int32 FHierarchicalLODUtilities::GetLODLevelForScreenAreaSize(UStaticMeshComponent* StaticMeshComponent, const float ScreenAreaSize)
+int32 FHierarchicalLODUtilities::GetLODLevelForScreenAreaSize(const UStaticMeshComponent* StaticMeshComponent, const float ScreenAreaSize)
 {
-	return ComputeStaticMeshLODLevel(StaticMeshComponent->StaticMesh->SourceModels, ScreenAreaSize);
+	check(StaticMeshComponent != nullptr);
+	const FStaticMeshRenderData* RenderData = StaticMeshComponent->StaticMesh->RenderData.GetOwnedPointer();
+	checkf(RenderData != nullptr, TEXT("StaticMesh in StaticMeshComponent %s contains invalid render data"), *StaticMeshComponent->GetName());
+	checkf(StaticMeshComponent->StaticMesh->SourceModels.Num() > 0, TEXT("StaticMesh in StaticMeshComponent %s contains no SourceModels"), *StaticMeshComponent->GetName());
+
+	return ComputeStaticMeshLODLevel(StaticMeshComponent->StaticMesh->SourceModels, RenderData, ScreenAreaSize);
 }
 
 AHierarchicalLODVolume* FHierarchicalLODUtilities::CreateVolumeForLODActor(ALODActor* InLODActor, UWorld* InWorld)
@@ -505,7 +753,7 @@ FHierarchicalLODProxyProcessor* FHierarchicalLODUtilities::GetProxyProcessor()
 	{
 		ProxyProcessor = new FHierarchicalLODProxyProcessor();
 	}
-	
+
 	return ProxyProcessor;
 }
 
diff --git a/Engine/Source/Developer/HierarchicalLODUtilities/Public/HierarchicalLODUtilities.h b/Engine/Source/Developer/HierarchicalLODUtilities/Public/HierarchicalLODUtilities.h
--- a/Engine/Source/Developer/HierarchicalLODUtilities/Public/HierarchicalLODUtilities.h
+++ b/Engine/Source/Developer/HierarchicalLODUtilities/Public/HierarchicalLODUtilities.h
@@ -65,23 +65,68 @@ public:
 	/** Returns the ALODActor parent for the given InActor, nullptr if none available */
 	static ALODActor* GetParentLODActor(const AActor* InActor);
 
-	/** Deletes the given actor's data and instance in the world */
-	static void DeleteLODActor(ALODActor* InActor);
+	/** Deletes the given cluster's data and instance in the world */
+	static void DestroyCluster(ALODActor* InActor);
 
-	/** Deletes the given actor's assets */
-	static void DeleteLODActorAssets(ALODActor* InActor);
+	/** Deletes the given cluster's assets */
+	static void DestroyClusterData(ALODActor* InActor);
 
 	/** Creates a new cluster actor in the given InWorld with InLODLevel as HLODLevel */
 	static ALODActor* CreateNewClusterActor(UWorld* InWorld, const int32 InLODLevel, AWorldSettings* WorldSettings);
 
+	/** Creates a new cluster in InWorld with InActors as sub actors*/
+	static ALODActor* CreateNewClusterFromActors(UWorld* InWorld, AWorldSettings* WorldSettings, const TArray<AActor*>& InActors, const int32 InLODLevel = 0);
+
+	/** Removes the given actor from it's parent cluster */
+	static const bool RemoveActorFromCluster(AActor* InActor);
+
+	/** Adds an actor to the given cluster*/
+	static const bool AddActorToCluster(AActor* InActor, ALODActor* InParentActor);
+
+	/** Merges two clusters together */
+	static const bool MergeClusters(ALODActor* TargetCluster, ALODActor* SourceCluster);
+
+	/** Checks if all actors have the same outer world */
+	static const bool AreActorsInSamePersistingLevel(const TArray<AActor*>& InActors);
+
+	/** Checks if all clusters are in the same HLOD level */
+	static const bool AreClustersInSameHLODLevel(const TArray<ALODActor*>& InLODActors);
+
+	/** Checks if all actors are in the same HLOD level */
+	static const bool AreActorsInSameHLODLevel(const TArray<AActor*>& InActors);
+
+	/** Checks if all actors are part of a cluster */
+	static const bool AreActorsClustered(const TArray<AActor*>& InActors);
+
+	/** Checks if an actor is clustered*/
+	static const bool IsActorClustered(const AActor* InActor);
+
+	/** Excludes an actor from the cluster generation process */
+	static void ExcludeActorFromClusterGeneration(AActor* InActor);
+
+	/**
+	* Destroys an LODActor instance
+	*
+	* @param InActor - ALODActor to destroy
+	*/
+	static void DestroyLODActor(ALODActor* InActor);
+
+	/**
+	* Extracts all the Static Mesh Actors from the given LODActor's SubActors array
+	*
+	* @param LODActor - LODActors to check the SubActors array for
+	* @param InOutActors - Array to fill with Static Mesh Actors
+	*/
+	static void ExtractStaticMeshActorsFromLODActor(ALODActor* LODActor, TArray<AActor*> &InOutActors);
+
 	/** Deletes all the ALODActors with the given HLODLevelIndex inside off InWorld */
 	static void DeleteLODActorsInHLODLevel(UWorld* InWorld, const int32 HLODLevelIndex);
 
 	/** Computes which LOD level of a Mesh corresponds to the given Distance (calculates closest screensize with distance) */
-	static int32 ComputeStaticMeshLODLevel(const TArray<FStaticMeshSourceModel>& SourceModels, const float ScreenAreaSize);
+	static int32 ComputeStaticMeshLODLevel(const TArray<FStaticMeshSourceModel>& SourceModels, const FStaticMeshRenderData* RenderData, const float ScreenAreaSize);
 
 	/** Computes the LODLevel for a StaticMeshComponent taking into account the ScreenArea */
-	static int32 GetLODLevelForScreenAreaSize(UStaticMeshComponent* StaticMeshComponent, const float ScreenAreaSize);
+	static int32 GetLODLevelForScreenAreaSize(const UStaticMeshComponent* StaticMeshComponent, const float ScreenAreaSize);
 	
 	/**
 	* Creates a HierarchicalLODVolume using the bounds of a given LODActor
diff --git a/Engine/Source/Developer/MaterialUtilities/Private/MaterialUtilities.cpp b/Engine/Source/Developer/MaterialUtilities/Private/MaterialUtilities.cpp
--- a/Engine/Source/Developer/MaterialUtilities/Private/MaterialUtilities.cpp
+++ b/Engine/Source/Developer/MaterialUtilities/Private/MaterialUtilities.cpp
@@ -821,7 +821,7 @@ bool FMaterialUtilities::ExportLandscapeMaterial(ALandscapeProxy* InLandscape, c
 	return true;
 }
 
-UMaterial* FMaterialUtilities::CreateMaterial(const FFlattenMaterial& InFlattenMaterial, UPackage* InOuter, const FString& BaseName, EObjectFlags Flags, const struct FMaterialProxySettings& MaterialProxySettings, TArray<UObject*>& OutGeneratedAssets)
+UMaterial* FMaterialUtilities::CreateMaterial(const FFlattenMaterial& InFlattenMaterial, UPackage* InOuter, const FString& BaseName, EObjectFlags Flags, const struct FMaterialProxySettings& MaterialProxySettings, TArray<UObject*>& OutGeneratedAssets, const TextureGroup& InTextureGroup /*= TEXTUREGROUP_World*/)
 {
 	// Base name for a new assets
 	// In case outer is null BaseName has to be long package name
@@ -859,7 +859,7 @@ UMaterial* FMaterialUtilities::CreateMaterial(const FFlattenMaterial& InFlattenM
 		const FString AssetName = TEXT("T_") + AssetBaseName + TEXT("_D");
 		const FString AssetLongName = AssetBasePath + AssetName;
 		const bool bSRGB = true;
-		UTexture2D* Texture = CreateTexture(InOuter, AssetLongName, InFlattenMaterial.DiffuseSize, InFlattenMaterial.DiffuseSamples, TC_Default, TEXTUREGROUP_World, Flags, bSRGB);
+		UTexture2D* Texture = CreateTexture(InOuter, AssetLongName, InFlattenMaterial.DiffuseSize, InFlattenMaterial.DiffuseSamples, TC_Default, InTextureGroup, Flags, bSRGB);
 		OutGeneratedAssets.Add(Texture);
 			
 		auto BasecolorExpression = NewObject<UMaterialExpressionTextureSample>(Material);
@@ -937,7 +937,7 @@ UMaterial* FMaterialUtilities::CreateMaterial(const FFlattenMaterial& InFlattenM
 
 		const FString AssetName = TEXT("T_") + AssetBaseName + TEXT("_MRS");
 		const bool bSRGB = true;
-		UTexture2D* Texture = CreateTexture(InOuter, AssetBasePath + AssetName, MergedSize, MergedSamples, TC_Default, TEXTUREGROUP_World, Flags, bSRGB);
+		UTexture2D* Texture = CreateTexture(InOuter, AssetBasePath + AssetName, MergedSize, MergedSamples, TC_Default, InTextureGroup, Flags, bSRGB);
 		OutGeneratedAssets.Add(Texture);
 
 		auto MergedExpression = NewObject<UMaterialExpressionTextureSample>(Material);
@@ -989,7 +989,7 @@ UMaterial* FMaterialUtilities::CreateMaterial(const FFlattenMaterial& InFlattenM
 		{
 			const FString AssetName = TEXT("T_") + AssetBaseName + TEXT("_M");
 			const bool bSRGB = true;
-			UTexture2D* Texture = CreateTexture(InOuter, AssetBasePath + AssetName, InFlattenMaterial.MetallicSize, InFlattenMaterial.MetallicSamples, TC_Default, TEXTUREGROUP_World, Flags, bSRGB);
+			UTexture2D* Texture = CreateTexture(InOuter, AssetBasePath + AssetName, InFlattenMaterial.MetallicSize, InFlattenMaterial.MetallicSamples, TC_Default, InTextureGroup, Flags, bSRGB);
 			OutGeneratedAssets.Add(Texture);
 
 			auto MetallicExpression = NewObject<UMaterialExpressionTextureSample>(Material);
@@ -1008,7 +1008,7 @@ UMaterial* FMaterialUtilities::CreateMaterial(const FFlattenMaterial& InFlattenM
 		{
 			const FString AssetName = TEXT("T_") + AssetBaseName + TEXT("_S");
 			const bool bSRGB = true;
-			UTexture2D* Texture = CreateTexture(InOuter, AssetBasePath + AssetName, InFlattenMaterial.SpecularSize, InFlattenMaterial.SpecularSamples, TC_Default, TEXTUREGROUP_World, Flags, bSRGB);
+			UTexture2D* Texture = CreateTexture(InOuter, AssetBasePath + AssetName, InFlattenMaterial.SpecularSize, InFlattenMaterial.SpecularSamples, TC_Default, InTextureGroup, Flags, bSRGB);
 			OutGeneratedAssets.Add(Texture);
 
 			auto SpecularExpression = NewObject<UMaterialExpressionTextureSample>(Material);
@@ -1027,7 +1027,7 @@ UMaterial* FMaterialUtilities::CreateMaterial(const FFlattenMaterial& InFlattenM
 		{
 			const FString AssetName = TEXT("T_") + AssetBaseName + TEXT("_R");
 			const bool bSRGB = true;
-			UTexture2D* Texture = CreateTexture(InOuter, AssetBasePath + AssetName, InFlattenMaterial.RoughnessSize, InFlattenMaterial.RoughnessSamples, TC_Default, TEXTUREGROUP_World, Flags, bSRGB);
+			UTexture2D* Texture = CreateTexture(InOuter, AssetBasePath + AssetName, InFlattenMaterial.RoughnessSize, InFlattenMaterial.RoughnessSamples, TC_Default, InTextureGroup, Flags, bSRGB);
 			OutGeneratedAssets.Add(Texture);
 
 			auto RoughnessExpression = NewObject<UMaterialExpressionTextureSample>(Material);
@@ -1085,7 +1085,7 @@ UMaterial* FMaterialUtilities::CreateMaterial(const FFlattenMaterial& InFlattenM
 	{
 		const FString AssetName = TEXT("T_") + AssetBaseName + TEXT("_N");
 		const bool bSRGB = false;
-		UTexture2D* Texture = CreateTexture(InOuter, AssetBasePath + AssetName, InFlattenMaterial.NormalSize, InFlattenMaterial.NormalSamples, TC_Normalmap, TEXTUREGROUP_WorldNormalMap, Flags, bSRGB);
+		UTexture2D* Texture = CreateTexture(InOuter, AssetBasePath + AssetName, InFlattenMaterial.NormalSize, InFlattenMaterial.NormalSamples, TC_Normalmap, (InTextureGroup != TEXTUREGROUP_World) ? InTextureGroup : TEXTUREGROUP_WorldNormalMap, Flags, bSRGB);
 		OutGeneratedAssets.Add(Texture);
 			
 		auto NormalExpression = NewObject<UMaterialExpressionTextureSample>(Material);
@@ -1120,7 +1120,7 @@ UMaterial* FMaterialUtilities::CreateMaterial(const FFlattenMaterial& InFlattenM
 	{
 		const FString AssetName = TEXT("T_") + AssetBaseName + TEXT("_E");
 		const bool bSRGB = true;
-		UTexture2D* Texture = CreateTexture(InOuter, AssetBasePath + AssetName, InFlattenMaterial.EmissiveSize, InFlattenMaterial.EmissiveSamples, TC_Default, TEXTUREGROUP_World, Flags, bSRGB);
+		UTexture2D* Texture = CreateTexture(InOuter, AssetBasePath + AssetName, InFlattenMaterial.EmissiveSize, InFlattenMaterial.EmissiveSamples, TC_Default, InTextureGroup, Flags, bSRGB);
 		OutGeneratedAssets.Add(Texture);
 
 		//Assign emissive color to the material
@@ -1159,7 +1159,7 @@ UMaterial* FMaterialUtilities::CreateMaterial(const FFlattenMaterial& InFlattenM
 	{
 		const FString AssetName = TEXT("T_") + AssetBaseName + TEXT("_O");
 		const bool bSRGB = true;
-		UTexture2D* Texture = CreateTexture(InOuter, AssetBasePath + AssetName, InFlattenMaterial.OpacitySize, InFlattenMaterial.OpacitySamples, TC_Default, TEXTUREGROUP_World, Flags, bSRGB);
+		UTexture2D* Texture = CreateTexture(InOuter, AssetBasePath + AssetName, InFlattenMaterial.OpacitySize, InFlattenMaterial.OpacitySamples, TC_Default, InTextureGroup, Flags, bSRGB);
 		OutGeneratedAssets.Add(Texture);
 
 		//Assign opacity to the material
@@ -1198,7 +1198,7 @@ UMaterial* FMaterialUtilities::CreateMaterial(const FFlattenMaterial& InFlattenM
 	{
 		const FString AssetName = TEXT("T_") + AssetBaseName + TEXT("_SSC");
 		const bool bSRGB = true;
-		UTexture2D* Texture = CreateTexture(InOuter, AssetBasePath + AssetName, InFlattenMaterial.SubSurfaceSize, InFlattenMaterial.SubSurfaceSamples, TC_Default, TEXTUREGROUP_World, Flags, bSRGB);
+		UTexture2D* Texture = CreateTexture(InOuter, AssetBasePath + AssetName, InFlattenMaterial.SubSurfaceSize, InFlattenMaterial.SubSurfaceSamples, TC_Default, InTextureGroup, Flags, bSRGB);
 		OutGeneratedAssets.Add(Texture);
 
 		//Assign emissive color to the material
@@ -1239,7 +1239,7 @@ UTexture2D* FMaterialUtilities::CreateTexture(UPackage* Outer, const FString& As
 	UTexture2D* Texture = FImageUtils::CreateTexture2D(Size.X, Size.Y, Samples, Outer, FPackageName::GetShortName(AssetLongName), Flags, TexParams);
 	Texture->LODGroup = LODGroup;
 	Texture->PostEditChange();
-		
+			
 	return Texture;
 }
 
@@ -1521,7 +1521,7 @@ void FMaterialUtilities::RemapUniqueMaterialIndices(const TArray<UMaterialInterf
 					// Mesh with vertex data baking, and with vertex colors - don't share materials at all.
 					for (int32 LocalMaterialIndex = 0; LocalMaterialIndex < MeshMaterialMap.Num(); LocalMaterialIndex++)
 					{
-						FMeshMaterialData Data(InMaterials[MeshMaterialMap[LocalMaterialIndex]], StaticMesh, false);
+						FMeshMaterialData Data(InMaterials[MeshMaterialMap[LocalMaterialIndex]], StaticMesh, true);
 						int32 Index = MeshMaterialData.Add(Data);
 						NewMeshMaterialMap.Add(Index);
 					}
@@ -1710,10 +1710,7 @@ bool FMaterialUtilities::ExportMaterial(struct FMaterialMergeData& InMaterialDat
 		// Purpose of the following line: use compiled material cached from previous call.
 		Exchange(ProxyCache, InMaterialData.ProxyCache);
 	}
-
-	// Make sure all the materials are loaded
-	FullyLoadMaterialStatic(Material);
-
+	
 	// Precache all used textures, otherwise could get everything rendered with low-res textures.
 	TArray<UTexture*> MaterialTextures;
 	Material->GetUsedTextures(MaterialTextures, EMaterialQualityLevel::Num, true, GMaxRHIFeatureLevel, true);
@@ -1828,8 +1825,8 @@ bool FMaterialUtilities::RenderMaterialPropertyToTexture(struct FMaterialMergeDa
 	check(CurrentlyRendering == false);
 	CurrentlyRendering = true;
 
-	UTextureRenderTarget2D* RenderTarget = CreateRenderTarget(bInForceLinearGamma, InPixelFormat, OutSampleSize);
-	
+	const bool bNormalMap = (InMaterialProperty == MP_Normal);
+	UTextureRenderTarget2D* RenderTarget = CreateRenderTarget(bInForceLinearGamma, bNormalMap, InPixelFormat, OutSampleSize);
 	OutSamples.Empty(InTargetSize.X * InTargetSize.Y);
 	bool bResult = FMeshRenderer::RenderMaterial(
 		InMaterialData,
@@ -1890,8 +1887,10 @@ bool FMaterialUtilities::RenderMaterialPropertyToTexture(struct FMaterialMergeDa
 	return bResult;
 }
 
-UTextureRenderTarget2D* FMaterialUtilities::CreateRenderTarget(bool bInForceLinearGamma, EPixelFormat InPixelFormat, FIntPoint& InTargetSize)
+UTextureRenderTarget2D* FMaterialUtilities::CreateRenderTarget(bool bInForceLinearGamma, bool bNormalMap, EPixelFormat InPixelFormat, FIntPoint& InTargetSize)
 {
+	const FLinearColor ClearColour = bNormalMap ? FLinearColor(0.0f, 0.0f, 0.0f, 0.0f) : FLinearColor(1.0f, 0.0f, 1.0f, 0.0f);
+	
 	// Find any pooled render target with suitable properties.
 	for (int32 RTIndex = 0; RTIndex < RenderTargetPool.Num(); RTIndex++)
 	{
@@ -1899,7 +1898,8 @@ UTextureRenderTarget2D* FMaterialUtilities::CreateRenderTarget(bool bInForceLine
 		if (RenderTarget->SizeX == InTargetSize.X &&
 			RenderTarget->SizeY == InTargetSize.Y &&
 			RenderTarget->OverrideFormat == InPixelFormat &&
-			RenderTarget->bForceLinearGamma == bInForceLinearGamma)
+			RenderTarget->bForceLinearGamma == bInForceLinearGamma &&
+			RenderTarget->ClearColor == RenderTarget->ClearColor )
 		{
 			return RenderTarget;
 		}
@@ -1909,7 +1909,7 @@ UTextureRenderTarget2D* FMaterialUtilities::CreateRenderTarget(bool bInForceLine
 	UTextureRenderTarget2D* NewRenderTarget = NewObject<UTextureRenderTarget2D>();
 	check(NewRenderTarget);
 	NewRenderTarget->AddToRoot();
-	NewRenderTarget->ClearColor = FLinearColor(0.0f, 0.0f, 0.0f, 0.0f);
+	NewRenderTarget->ClearColor = ClearColour;
 	NewRenderTarget->TargetGamma = 0.0f;
 	NewRenderTarget->InitCustomFormat(InTargetSize.X, InTargetSize.Y, InPixelFormat, bInForceLinearGamma);
 
diff --git a/Engine/Source/Developer/MaterialUtilities/Private/MeshRendering.cpp b/Engine/Source/Developer/MaterialUtilities/Private/MeshRendering.cpp
--- a/Engine/Source/Developer/MaterialUtilities/Private/MeshRendering.cpp
+++ b/Engine/Source/Developer/MaterialUtilities/Private/MeshRendering.cpp
@@ -34,87 +34,108 @@
 #include "ThumbnailHelpers.h" // for FClassThumbnailScene
 #include "ShaderCompiler.h"   // for GShaderCompilingManager
 
+FColor BoxBlurSample(TArray<FColor>& InBMP, int32 X, int32 Y, int32 InImageWidth, int32 InImageHeight, bool bIsNormalMap)
+{
+	const int32 SampleCount = 8;
+	static int32 PixelIndices[SampleCount] = { -(InImageWidth + 1), -InImageWidth, -(InImageWidth - 1),
+									 -1, 1,
+									 (InImageWidth + -1), InImageWidth, (InImageWidth + 1) };
+
+	static int32 PixelOffsetX[SampleCount] = { -1, 0, 1,
+									-1, 1,
+									-1, 0, 1 };
+
+	int32 PixelsSampled = 0;
+	FLinearColor CombinedColor = FColor::Black;
+	
+	// Take samples for blur with square indices
+	for (int32 SampleIndex = 0; SampleIndex < SampleCount; ++SampleIndex)
+	{
+		const int32 PixelIndex = ((Y * InImageWidth) + X) + PixelIndices[SampleIndex];
+
+		const int32 XIndex = X + PixelOffsetX[SampleIndex];
+
+		// Check if we are not out of texture bounds
+		if (InBMP.IsValidIndex(PixelIndex) && XIndex >= 0 && XIndex < InImageWidth)
+		{
+			FLinearColor SampledColor = InBMP[PixelIndex].ReinterpretAsLinear();
+			// Check if the pixel is a rendered one (not clear colour)
+			if ((!(SampledColor.R == 1.0f && SampledColor.B == 1.0f && SampledColor.G == 0.0f)) && (!bIsNormalMap || SampledColor.B != 0.0f))
+			{
+				CombinedColor += SampledColor;
+				++PixelsSampled;
+			}
+		}
+	}
+	CombinedColor /= PixelsSampled;
+
+	if (PixelsSampled == 0)
+	{
+		return InBMP[((Y * InImageWidth) + X)];
+	}
+
+	return CombinedColor.ToFColor(false);
+}
 
-void PerformImageDilation(TArray<FColor>& InBMP, int32 InImageWidth, int32 InImageHeight, bool IsNormalMap)
+void PerformUVBorderSmear(TArray<FColor>& InBMP, int32 InImageWidth, int32 InImageHeight, bool IsNormalMap)
 {
-	int32 PixelIndex = 0;
-	int32 PixelIndices[16];
+	TArray<FColor> Swap;
+	Swap.Append(InBMP);
+
+	TArray<FColor>* Current = &InBMP;
+	TArray<FColor>* Scratch = &Swap;
+
+	bool bSwap = false;
+
+	int32 MagentaPixels = 1;
 
-	for (int32 Y = 0; Y < InImageHeight; Y++)
+	int32 LoopCount = 0;
+	const int32 MaxIterations = 32;
+	
+	// Sampling
+	while (MagentaPixels && (LoopCount <= MaxIterations))
 	{
-		for (int32 X = 0; X < InImageWidth; X++, PixelIndex++)
-		{
-			FColor& Color = InBMP[PixelIndex];
-			if (Color.A == 0 || (IsNormalMap && Color.B == 0))
+		MagentaPixels = 0;
+		// Left / right, Top / down
+		for (int32 Y = 0; Y < InImageHeight; Y++)
+		{		
+			for (int32 X = 0; X < InImageWidth; X++)
 			{
-				// process this pixel only if it is not filled with color; note: normal map has A==255, so detect
-				// empty pixels by zero blue component, which is not possible for normal maps
-				int32 NumPixelsToCheck = 0;
-				// check line at Y-1
-				if (Y > 0)
-				{
-					PixelIndices[NumPixelsToCheck++] = PixelIndex - InImageWidth; // X,Y-1
-					if (X > 0)
-					{
-						PixelIndices[NumPixelsToCheck++] = PixelIndex - InImageWidth - 1; // X-1,Y-1
-					}
-					if (X < InImageWidth - 1)
-					{
-						PixelIndices[NumPixelsToCheck++] = PixelIndex - InImageWidth + 1; // X+1,Y-1
-					}
-				}
-				// check line at Y
-				if (X > 0)
-				{
-					PixelIndices[NumPixelsToCheck++] = PixelIndex - 1; // X-1,Y
-				}
-				if (X < InImageWidth - 1)
-				{
-					PixelIndices[NumPixelsToCheck++] = PixelIndex + 1; // X+1,Y
-				}
-				// check line at Y+1
-				if (Y < InImageHeight - 1)
+				int32 PixelIndex = (Y * InImageWidth) + X;
+
+				FColor& Color = (*Current)[PixelIndex];
+				if ((Color.R == 255 && Color.B == 255 && Color.G == 0) || (IsNormalMap && Color.B == 0))
 				{
-					PixelIndices[NumPixelsToCheck++] = PixelIndex + InImageWidth; // X,Y+1
-					if (X > 0)
-					{
-						PixelIndices[NumPixelsToCheck++] = PixelIndex + InImageWidth - 1; // X-1,Y+1
-					}
-					if (X < InImageWidth - 1)
+					MagentaPixels++;
+					FColor SampledColor = BoxBlurSample(*Scratch, X, Y, InImageWidth, InImageHeight, IsNormalMap);
+					// If it's a valid pixel
+					if ((!(SampledColor.R == 255 && SampledColor.B == 255 && SampledColor.G == 0)) && (!IsNormalMap || SampledColor.B != 0))
 					{
-						PixelIndices[NumPixelsToCheck++] = PixelIndex + InImageWidth + 1; // X+1,Y+1
+						Color = SampledColor;
 					}
-				}
-				// get color
-				int32 BestColorValue = 0;
-				FColor BestColor(0);
-				for (int32 PixelToCheck = 0; PixelToCheck < NumPixelsToCheck; PixelToCheck++)
-				{
-					const FColor& ColorToCheck = InBMP[PixelIndices[PixelToCheck]];
-					if (ColorToCheck.A != 0 && (!IsNormalMap || ColorToCheck.B != 0))
+					else
 					{
-						// consider only original pixels, not computed ones
-						int32 ColorValue = ColorToCheck.R + ColorToCheck.G + ColorToCheck.B;
-						if (ColorValue > BestColorValue)
+						// If we are at the end of our iterations, replace the pixels with black
+						if (LoopCount == (MaxIterations - 1))
 						{
-							BestColorValue = ColorValue;
-							BestColor = ColorToCheck;
+							Color = FColor::Black;
 						}
 					}
 				}
-				// put the computed pixel back
-				if (BestColorValue != 0)
-				{
-					Color = BestColor;
-					Color.A = 0;
-#if VISUALIZE_DILATION
-					Color.R = 255;
-					Color.G = 0;
-					Color.B = 0;
-#endif
-				}
 			}
 		}
+
+		TArray<FColor>& Temp = *Scratch;
+		Scratch = Current;
+		Current = &Temp;
+
+		LoopCount++;
+	}
+
+	if (Current != &InBMP)
+	{
+		InBMP.Empty();
+		InBMP.Append(*Current);
 	}
 }
 
@@ -907,10 +928,12 @@ bool FMeshRenderer::RenderMaterial(struct FMaterialMergeData& InMaterialData, FM
 		InMaterialData.EmissiveScale = MaxValue;
 	}
 
-	PerformImageDilation(OutBMP, InRenderTarget->GetSurfaceWidth(), InRenderTarget->GetSurfaceHeight(), bNormalmap);
+	PerformUVBorderSmear(OutBMP, InRenderTarget->GetSurfaceWidth(), InRenderTarget->GetSurfaceHeight(), bNormalmap);
+
+//#define SAVE_INTERMEDIATE_TEXTURES 1
 #ifdef SAVE_INTERMEDIATE_TEXTURES
-	FString FilenameString = FString::Printf(
-		SAVE_INTERMEDIATE_TEXTURES TEXT("/%s-mat%d-prop%d.bmp"),
+	FilenameString = FString::Printf(
+		TEXT( "D:/TextureTest/%s-mat%d-prop%d.bmp"),
 		*InMaterialProxy->GetFriendlyName(), InMaterialData.MaterialIndex, (int32)InMaterialProperty);
 	FFileHelper::CreateBitmap(*FilenameString, InRenderTarget->GetSurfaceWidth(), InRenderTarget->GetSurfaceHeight(), OutBMP.GetData());
 #endif // SAVE_INTERMEDIATE_TEXTURES
diff --git a/Engine/Source/Developer/MaterialUtilities/Public/MaterialUtilities.h b/Engine/Source/Developer/MaterialUtilities/Public/MaterialUtilities.h
--- a/Engine/Source/Developer/MaterialUtilities/Public/MaterialUtilities.h
+++ b/Engine/Source/Developer/MaterialUtilities/Public/MaterialUtilities.h
@@ -269,7 +269,7 @@ public:
 	 * @param OutGeneratedAssets	List of generated assets - material, textures
 	 * @return						Returns a pointer to the constructed UMaterial object.
 	 */
-	static UMaterial* CreateMaterial(const FFlattenMaterial& InFlattenMaterial, UPackage* InOuter, const FString& BaseName, EObjectFlags Flags, const struct FMaterialProxySettings& MaterialProxySettings, TArray<UObject*>& OutGeneratedAssets);
+	static UMaterial* CreateMaterial(const FFlattenMaterial& InFlattenMaterial, UPackage* InOuter, const FString& BaseName, EObjectFlags Flags, const struct FMaterialProxySettings& MaterialProxySettings, TArray<UObject*>& OutGeneratedAssets, const TextureGroup& InTextureGroup = TEXTUREGROUP_World);
 	
 	/**
 	* Creates bakes textures for a ULandscapeComponent
@@ -358,7 +358,7 @@ private:
 	* @param InTargetSize			Dimensions of the render target
 	* @return						Created render target
 	*/
-	static UTextureRenderTarget2D* CreateRenderTarget(bool bInForceLinearGamma, EPixelFormat InPixelFormat, FIntPoint& InTargetSize);
+	static UTextureRenderTarget2D* CreateRenderTarget(bool bInForceLinearGamma, bool bNormalMap, EPixelFormat InPixelFormat, FIntPoint& InTargetSize );
 	
 	/** Clears the pool with available render targets */
 	static void ClearRenderTargetPool();	
diff --git a/Engine/Source/Developer/MeshUtilities/MeshUtilities.Build.cs b/Engine/Source/Developer/MeshUtilities/MeshUtilities.Build.cs
--- a/Engine/Source/Developer/MeshUtilities/MeshUtilities.Build.cs
+++ b/Engine/Source/Developer/MeshUtilities/MeshUtilities.Build.cs
@@ -26,7 +26,8 @@ public class MeshUtilities : ModuleRules
                 "MeshBoneReduction",		
                 "UnrealEd",
                 "RHI",                
-                "HierarchicalLODUtilities"
+                "HierarchicalLODUtilities",
+                "Landscape"
 			}
 		);
 
diff --git a/Engine/Source/Developer/MeshUtilities/Private/MeshUtilities.cpp b/Engine/Source/Developer/MeshUtilities/Private/MeshUtilities.cpp
--- a/Engine/Source/Developer/MeshUtilities/Private/MeshUtilities.cpp
+++ b/Engine/Source/Developer/MeshUtilities/Private/MeshUtilities.cpp
@@ -23,6 +23,11 @@
 #include "MeshMergeData.h"
 #include "Editor/EditorPerProjectUserSettings.h"
 
+#include "Landscape.h"
+#include "LandscapeProxy.h"
+#include "LandscapeHeightfieldCollisionComponent.h"
+#include "Engine/HLODMeshCullingVolume.h"
+
 //@todo - implement required vector intrinsics for other implementations
 #if PLATFORM_ENABLE_VECTORINTRINSICS
 #include "kDOP.h"
@@ -166,7 +171,8 @@ private:
 		const float ScreenAreaSize,
 		bool bSilent = false) const override;
 
-	virtual void CreateProxyMesh(const TArray<AActor*>& InActors, const struct FMeshProxySettings& InMeshProxySettings, UPackage* InOuter, const FString& InProxyBasePackageName, const FGuid InGuid, FCreateProxyDelegate InProxyCreatedDelegate, const bool bAllowAsync) override;
+	virtual void CreateProxyMesh(const TArray<AActor*>& InActors, const struct FMeshProxySettings& InMeshProxySettings, UPackage* InOuter, const FString& InProxyBasePackageName, const FGuid InGuid, FCreateProxyDelegate InProxyCreatedDelegate, const bool bAllowAsync,
+		const float ScreenAreaSize = 1.0f) override;
 
 	DEPRECATED(4.11, "Please use CreateProxyMesh with new signature")
 		virtual void CreateProxyMesh(
@@ -183,15 +189,15 @@ private:
 		const struct FMeshProxySettings& InProxySettings,
 		UPackage* InOuter,
 		const FString& ProxyBasePackageName,
-		TArray<UObject*>& OutAssetsToSync) override;
-
+		TArray<UObject*>& OutAssetsToSync,
+		const float ScreenAreaSize = 1.0f) override;
 
 	virtual void FlattenMaterialsWithMeshData(TArray<UMaterialInterface*>& InMaterials, TArray<FRawMeshExt>& InSourceMeshes, TMap<FMeshIdAndLOD, TArray<int32>>& InMaterialIndexMap, TArray<bool>& InMeshShouldBakeVertexData, const FMaterialProxySettings &InMaterialProxySettings, TArray<FFlattenMaterial> &OutFlattenedMaterials) const override;
 
-
 	bool ConstructRawMesh(
-		UStaticMeshComponent* MeshComponent,
-		int32 LODIndex,
+		const UStaticMeshComponent* InMeshComponent,
+		int32 InLODIndex,
+		const bool bPropagateVertexColours,
 		FRawMesh& OutRawMesh,
 		TArray<UMaterialInterface*>& OutUniqueMaterials,
 		TArray<int32>& OutGlobalMaterialIndices
@@ -199,7 +205,7 @@ private:
 
 	virtual void ExtractMeshDataForGeometryCache(FRawMesh& RawMesh, const FMeshBuildSettings& BuildSettings, TArray<FStaticMeshBuildVertex>& OutVertices, TArray<TArray<uint32> >& OutPerSectionIndices);
 
-	virtual bool PropagatePaintedColorsToRawMesh(UStaticMeshComponent* StaticMeshComponent, int32 LODIndex, FRawMesh& RawMesh) const override;
+	virtual bool PropagatePaintedColorsToRawMesh(const UStaticMeshComponent* StaticMeshComponent, int32 LODIndex, FRawMesh& RawMesh) const override;
 
 	virtual void CalculateTextureCoordinateBoundsForRawMesh(const FRawMesh& InRawMesh, TArray<FBox2D>& OutBounds) const override;
 
@@ -324,7 +330,7 @@ protected:
 		FMaterialUtilities::OptimizeFlattenMaterial(Data->Material);
 
 		// Construct proxy material
-		UMaterial* ProxyMaterial = FMaterialUtilities::CreateMaterial(Data->Material, Data->MergeData->InOuter, Data->MergeData->ProxyBasePackageName, RF_Public | RF_Standalone, Data->MergeData->InProxySettings.MaterialSettings, OutAssetsToSync);
+		UMaterial* ProxyMaterial = FMaterialUtilities::CreateMaterial(Data->Material, Data->MergeData->InOuter, Data->MergeData->ProxyBasePackageName, RF_Public | RF_Standalone, Data->MergeData->InProxySettings.MaterialSettings, OutAssetsToSync, TEXTUREGROUP_HierarchicalLOD);
 
 		// Set material static lighting usage flag if project has static lighting enabled
 		static const auto AllowStaticLightingVar = IConsoleManager::Get().FindTConsoleVariableDataInt(TEXT("r.AllowStaticLighting"));
@@ -5177,19 +5183,45 @@ static FIntPoint ConditionalImageResize(const FIntPoint& SrcSize, const FIntPoin
 	return SrcSize;
 }
 
+static void RetrieveValidStaticMeshComponentsForMerging(AActor* InActor, TArray<UStaticMeshComponent*>& OutComponents)
+{
+	TInlineComponentArray<UStaticMeshComponent*> Components;
+	InActor->GetComponents<UStaticMeshComponent>(Components);
+	// TODO: support derived classes from static component
+	Components.RemoveAll([](UStaticMeshComponent* Val){ return !(Val->IsA(UStaticMeshComponent::StaticClass()) || Val->IsA(USplineMeshComponent::StaticClass())); });
+
+	// TODO: support non-opaque materials
+	//Components.RemoveAll(&NonOpaqueMaterialPredicate);
+	OutComponents.Append(Components);
+}
 
 
-void FMeshUtilities::CreateProxyMesh(const TArray<AActor*>& InActors, const struct FMeshProxySettings& InMeshProxySettings, UPackage* InOuter, const FString& InProxyBasePackageName, const FGuid InGuid, FCreateProxyDelegate InProxyCreatedDelegate, const bool bAllowAsync)
+void FMeshUtilities::CreateProxyMesh(const TArray<AActor*>& InActors, const struct FMeshProxySettings& InMeshProxySettings, UPackage* InOuter, const FString& InProxyBasePackageName, const FGuid InGuid, FCreateProxyDelegate InProxyCreatedDelegate, const bool bAllowAsync, const float ScreenAreaSize)
 {
 	FScopedSlowTask MainTask(100, (LOCTEXT("MeshUtilities_CreateProxyMesh", "Creating Proxy Mesh")));
 	MainTask.MakeDialog();
 
+	// Error/warning checking for input
 	if (MeshMerging == NULL)
 	{
 		UE_LOG(LogMeshUtilities, Log, TEXT("No automatic mesh merging module available"));
 		return;
 	}
 
+	// Check that the delegate has a func-ptr bound to it
+	if (!InProxyCreatedDelegate.IsBound())
+	{
+		UE_LOG(LogMeshUtilities, Log, TEXT("Invalid (unbound) delegate for returning generated proxy mesh"));
+		return;
+	}
+
+	// No actors given as input
+	if (InActors.Num() == 0)
+	{
+		UE_LOG(LogMeshUtilities, Log, TEXT("No actors specified to generate a proxy mesh for"));
+		return;
+	}
+
 	// Base asset name for a new assets
 	// In case outer is null ProxyBasePackageName has to be long package name
 	if (InOuter == nullptr && FPackageName::IsShortPackageName(InProxyBasePackageName))
@@ -5197,105 +5229,102 @@ void FMeshUtilities::CreateProxyMesh(const TArray<AActor*>& InActors, const stru
 		UE_LOG(LogMeshUtilities, Warning, TEXT("Invalid long package name: '%s'."), *InProxyBasePackageName);
 		return;
 	}
-
-	const FString AssetBaseName = FPackageName::GetShortName(InProxyBasePackageName);
-	const FString AssetBasePath = InOuter ? TEXT("") : FPackageName::GetLongPackagePath(InProxyBasePackageName) + TEXT("/");
-	UWorld* InWorld = InActors.Num() ? InActors[0]->GetWorld() : nullptr;
-
-	TArray<UStaticMeshComponent*> ComponentsToMerge;
-
+	
 	MainTask.EnterProgressFrame(10.0f);
 
+	// Retrieve static mesh components valid for merging from the given set of actors	
+	TArray<UStaticMeshComponent*> ComponentsToMerge;
 	{
 		FScopedSlowTask SubTask(InActors.Num(), (LOCTEXT("MeshUtilities_CreateProxyMesh_CollectStaticMeshComponents", "Collecting StaticMeshComponents")));
 		// Collect components to merge
 		for (AActor* Actor : InActors)
 		{
-			TInlineComponentArray<UStaticMeshComponent*> Components;
-			Actor->GetComponents<UStaticMeshComponent>(Components);
-			// TODO: support derived classes from static component
-			Components.RemoveAll([](UStaticMeshComponent* Val){ return !(Val->IsA(UStaticMeshComponent::StaticClass()) || Val->IsA(USplineMeshComponent::StaticClass())); });
-
-			// TODO: support non-opaque materials
-			//Components.RemoveAll(&NonOpaqueMaterialPredicate);
-			//
-			ComponentsToMerge.Append(Components);
-
+			RetrieveValidStaticMeshComponentsForMerging(Actor, ComponentsToMerge);
 			SubTask.EnterProgressFrame(1.0f);
 		}
 	}
 
 	MainTask.EnterProgressFrame(10.0f);
 
-	// Retrieve mesh / material data
-	TArray<FRawMeshExt> SourceMeshes;
-	SourceMeshes.AddZeroed(ComponentsToMerge.Num());
-
+	// Check if there are actually any static mesh components to merge
+	if (ComponentsToMerge.Num() == 0)
+	{
+		UE_LOG(LogMeshUtilities, Log, TEXT("No valid static mesh components found in given set of Actors"));
+		return;
+	}
 
+	
 	typedef FIntPoint FMeshIdAndLOD;
-	TArray<UMaterialInterface*> UniqueMaterials;
-	TMap<FMeshIdAndLOD, TArray<int32>> MaterialMap;
+	TArray<FRawMeshExt> SourceMeshes;
+	TArray<UMaterialInterface*> GlobalUniqueMaterialList;
+	TMap<FMeshIdAndLOD, TArray<int32>> GlobalMaterialMap;
+	static const int32 ProxyMeshTargetLODLevel = 0;
 
-	bool bVertexColorsInMeshes = false;
-	bool bOcuppiedUVChannels[MAX_MESH_TEXTURE_COORDS] = {};
+	FBoxSphereBounds EstimatedBounds(ForceInitToZero);
+	for (const UStaticMeshComponent* StaticMeshComponent : ComponentsToMerge)
+	{
+		EstimatedBounds = EstimatedBounds + StaticMeshComponent->Bounds;
+	}
 
-	for (int32 ComponentIndex = 0; ComponentIndex < ComponentsToMerge.Num(); ComponentIndex++)
+	static const float FOVRad = 90.0f * (float)PI / 360.0f;
+	static const FMatrix ProjectionMatrix = FPerspectiveMatrix(FOVRad, 1920, 1080, 0.01f);
+	float EstimatedDistance = FHierarchicalLODUtilities::CalculateDrawDistanceFromScreenSize(EstimatedBounds.SphereRadius, ScreenAreaSize, ProjectionMatrix);
+	
+	// Retrieve mesh / material data
+	for (const UStaticMeshComponent* StaticMeshComponent : ComponentsToMerge)
 	{
-		UStaticMeshComponent* StaticMeshComponent = ComponentsToMerge[ComponentIndex];
-		TArray<int32> MeshMaterialMap;
+		TArray<int32> StaticMeshGlobalMaterialMap;
+		FRawMesh RawMesh;		
+		
+		const int32 ProxyMeshSourceLODLevel = InMeshProxySettings.bCalculateLODSourceModels ? FHierarchicalLODUtilities::GetLODLevelForScreenAreaSize(StaticMeshComponent, FHierarchicalLODUtilities::CalculateScreenSizeFromDrawDistance(StaticMeshComponent->Bounds.SphereRadius, ProjectionMatrix, EstimatedDistance)) : 0;
+		// Proxy meshes should always propagate vertex colours for material baking
+		static const bool bPropagateVertexColours = true;
 
-		SourceMeshes[ComponentIndex].MeshLODData[0].SourceStaticMesh = StaticMeshComponent->StaticMesh;
-		FRawMesh& RawMesh = SourceMeshes[ComponentIndex].MeshLODData[0].RawMesh;
+		const bool bValidRawMesh = ConstructRawMesh(StaticMeshComponent, ProxyMeshSourceLODLevel, bPropagateVertexColours, RawMesh, GlobalUniqueMaterialList, StaticMeshGlobalMaterialMap);
 
-		if (ConstructRawMesh(StaticMeshComponent, 0, RawMesh, UniqueMaterials, MeshMaterialMap))
+		if ( bValidRawMesh )
 		{
-			MaterialMap.Add(FMeshIdAndLOD(ComponentIndex, 0), MeshMaterialMap);
+			// Add constructed raw mesh to source mesh array
+			const int32 SourceMeshIndex = SourceMeshes.AddZeroed();
+			SourceMeshes[SourceMeshIndex].MeshLODData[ProxyMeshTargetLODLevel].RawMesh = RawMesh;
 
-			// Should always bake vertex data for proxies			
-			// Propagate painted vertex colors into our raw mesh
-			PropagatePaintedColorsToRawMesh(StaticMeshComponent, 0, RawMesh);
-			// Whether at least one of the meshes has vertex colors
-			bVertexColorsInMeshes |= (RawMesh.WedgeColors.Num() != 0);
-
-			// Which UV channels has data at least in one mesh
-			for (int32 ChannelIdx = 0; ChannelIdx < MAX_MESH_TEXTURE_COORDS; ++ChannelIdx)
-			{
-				bOcuppiedUVChannels[ChannelIdx] |= (RawMesh.WedgeTexCoords[ChannelIdx].Num() != 0);
-			}
+			// Append retrieved materials for this static mesh component to the global material map
+			GlobalMaterialMap.Add(FMeshIdAndLOD(SourceMeshIndex, ProxyMeshTargetLODLevel), StaticMeshGlobalMaterialMap);
 		}
 	}
 
 	if (SourceMeshes.Num() == 0)
 	{
+		UE_LOG(LogMeshUtilities, Log, TEXT("No valid (or completely culled) raw meshes constructed from static mesh components"));
 		return;
 	}
 
 	TArray<bool> MeshShouldBakeVertexData;
-	TMap<FMeshIdAndLOD, TArray<int32> > NewMaterialMap;
-	TArray<UMaterialInterface*> NewStaticMeshMaterials;
+	TMap<FMeshIdAndLOD, TArray<int32> > NewGlobalMaterialMap;
+	TArray<UMaterialInterface*> NewGlobalUniqueMaterialList;
 	FMaterialUtilities::RemapUniqueMaterialIndices(
-		UniqueMaterials,
+		GlobalUniqueMaterialList,
 		SourceMeshes,
-		MaterialMap,
+		GlobalMaterialMap,
 		InMeshProxySettings.MaterialSettings,
 		true,
 		true,
 		MeshShouldBakeVertexData,
-		NewMaterialMap,
-		NewStaticMeshMaterials);
+		NewGlobalMaterialMap,
+		NewGlobalUniqueMaterialList);
 	// Use shared material data.
-	Exchange(MaterialMap, NewMaterialMap);
-	Exchange(UniqueMaterials, NewStaticMeshMaterials);
+	Exchange(GlobalMaterialMap, NewGlobalMaterialMap);
+	Exchange(GlobalUniqueMaterialList, NewGlobalUniqueMaterialList);
 
 	// Flatten Materials
 	TArray<FFlattenMaterial> FlattenedMaterials;
-	FlattenMaterialsWithMeshData(UniqueMaterials, SourceMeshes, MaterialMap, MeshShouldBakeVertexData, InMeshProxySettings.MaterialSettings, FlattenedMaterials);
+	FlattenMaterialsWithMeshData(GlobalUniqueMaterialList, SourceMeshes, GlobalMaterialMap, MeshShouldBakeVertexData, InMeshProxySettings.MaterialSettings, FlattenedMaterials);
 
 	//For each raw mesh, re-map the material indices from Local to Global material indices space
 	for (int32 RawMeshIndex = 0; RawMeshIndex < SourceMeshes.Num(); ++RawMeshIndex)
 	{
-		const TArray<int32>& GlobalMaterialIndices = *MaterialMap.Find(FMeshIdAndLOD(RawMeshIndex, 0));
-		TArray<int32>& MaterialIndices = SourceMeshes[RawMeshIndex].MeshLODData[0].RawMesh.FaceMaterialIndices;
+		const TArray<int32>& GlobalMaterialIndices = *GlobalMaterialMap.Find(FMeshIdAndLOD(RawMeshIndex, ProxyMeshTargetLODLevel));
+		TArray<int32>& MaterialIndices = SourceMeshes[RawMeshIndex].MeshLODData[ProxyMeshTargetLODLevel].RawMesh.FaceMaterialIndices;
 		int32 MaterialIndicesCount = MaterialIndices.Num();
 
 		for (int32 TriangleIndex = 0; TriangleIndex < MaterialIndicesCount; ++TriangleIndex)
@@ -5321,11 +5350,11 @@ void FMeshUtilities::CreateProxyMesh(const TArray<AActor*>& InActors, const stru
 	// Add this proxy job to map	
 	Processor->AddProxyJob(InGuid, Data);
 
-	// We are only using LOD level 0
+	// We are only using LOD level 0 (ProxyMeshTargetLODLevel)
 	TArray<FMeshMergeData> MergeData;
 	for (FRawMeshExt& SourceMesh : SourceMeshes)
 	{
-		MergeData.Add(SourceMesh.MeshLODData[0]);
+		MergeData.Add(SourceMesh.MeshLODData[ProxyMeshTargetLODLevel]);
 	}
 
 	// Choose Simplygon Swarm (if available) or local proxy lod method
@@ -5345,7 +5374,7 @@ void FMeshUtilities::CreateProxyMesh(const TArray<AActor*>& Actors, const struct
 	CreateProxyMesh(Actors, InProxySettings, InOuter, ProxyBasePackageName, OutAssetsToSync);
 }
 
-void FMeshUtilities::CreateProxyMesh(const TArray<AActor*>& Actors, const struct FMeshProxySettings& InProxySettings, UPackage* InOuter, const FString& ProxyBasePackageName, TArray<UObject*>& OutAssetsToSync)
+void FMeshUtilities::CreateProxyMesh(const TArray<AActor*>& Actors, const struct FMeshProxySettings& InProxySettings, UPackage* InOuter, const FString& ProxyBasePackageName, TArray<UObject*>& OutAssetsToSync, const float ScreenAreaSize)
 {
 	FCreateProxyDelegate Delegate;
 
@@ -5360,7 +5389,7 @@ void FMeshUtilities::CreateProxyMesh(const TArray<AActor*>& Actors, const struct
 	}
 	);
 
-	CreateProxyMesh(Actors, InProxySettings, InOuter, ProxyBasePackageName, JobGuid, Delegate, false);
+	CreateProxyMesh(Actors, InProxySettings, InOuter, ProxyBasePackageName, JobGuid, Delegate, false, ScreenAreaSize);
 }
 
 void FMeshUtilities::FlattenMaterialsWithMeshData(TArray<UMaterialInterface*>& InMaterials, TArray<FRawMeshExt>& InSourceMeshes, TMap<FMeshIdAndLOD, TArray<int32>>& InMaterialIndexMap, TArray<bool>& InMeshShouldBakeVertexData, const FMaterialProxySettings &InMaterialProxySettings, TArray<FFlattenMaterial> &OutFlattenedMaterials) const
@@ -5393,10 +5422,10 @@ void FMeshUtilities::FlattenMaterialsWithMeshData(TArray<UMaterialInterface*>& I
 		int32 LocalMaterialIndex = 0;
 		int32 LocalTextureBoundIndex = 0;
 		FMeshMergeData* MergeData = nullptr;
-		for (int32 MeshIndex = 0; MeshIndex < InSourceMeshes.Num() && MergeData == nullptr; MeshIndex++)
+		for (int32 LODIndex = 0; LODIndex < MAX_STATIC_MESH_LODS; ++LODIndex)
 		{
 			bool bFound = false;
-			for (int32 LODIndex = 0; LODIndex < MAX_STATIC_MESH_LODS; ++LODIndex)
+			for (int32 MeshIndex = 0; MeshIndex < InSourceMeshes.Num() && MergeData == nullptr; MeshIndex++)
 			{
 				if (InSourceMeshes[MeshIndex].MeshLODData[LODIndex].RawMesh.VertexPositions.Num())
 				{
@@ -5617,118 +5646,270 @@ static void ExportStaticMeshLOD(const FStaticMeshLODResources& StaticMeshLOD, FR
 	}
 }
 
-bool FMeshUtilities::ConstructRawMesh(
-	UStaticMeshComponent* InMeshComponent,
-	int32 LODIndex,
-	FRawMesh& OutRawMesh,
-	TArray<UMaterialInterface*>& OutUniqueMaterials,
-	TArray<int32>& OutGlobalMaterialIndices) const
+const bool IsLandscapeHit(const FVector& RayOrigin, const FVector& RayEndPoint, const UWorld* World, const TArray<ALandscapeProxy*>& LandscapeProxies)
 {
-	UStaticMesh* SrcMesh = InMeshComponent->StaticMesh;
-	if (SrcMesh == NULL)
+	static FName TraceTag = FName(TEXT("LandscapeTrace"));
+	TArray<FHitResult> Results;
+	// Each landscape component has 2 collision shapes, 1 of them is specific to landscape editor
+	// Trace only ECC_Visibility channel, so we do hit only Editor specific shape
+	World->LineTraceMultiByObjectType(Results, RayOrigin, RayEndPoint, FCollisionObjectQueryParams(ECollisionChannel::ECC_Visibility), FCollisionQueryParams(TraceTag, true));
+	
+	bool bHitLandscape = false;
+	FVector OutHitLocation;
+	
+	for (const FHitResult& HitResult : Results)
 	{
-		UE_LOG(LogMeshUtilities, Warning, TEXT("No static mesh actor found in component %s."), *InMeshComponent->GetName());
-		return false;
+		ULandscapeHeightfieldCollisionComponent* CollisionComponent = Cast<ULandscapeHeightfieldCollisionComponent>(HitResult.Component.Get());
+		if (CollisionComponent)
+		{
+			ALandscapeProxy* HitLandscape = CollisionComponent->GetLandscapeProxy();
+			if (HitLandscape && LandscapeProxies.Contains(HitLandscape))
+			{
+				// Could write a correct clipping algorithm, that clips the triangle to hit location
+				OutHitLocation = HitLandscape->LandscapeActorToWorld().InverseTransformPosition(HitResult.Location);
+				// Above landscape so visible
+				bHitLandscape = true;
+			}
+		}
 	}
 
-	if (!SrcMesh->SourceModels.IsValidIndex(LODIndex))
-	{
-		UE_LOG(LogMeshUtilities, Log, TEXT("No mesh data found for LOD%d %s."), LODIndex, *SrcMesh->GetName());
-		return false;
-	}
+	return bHitLandscape;
+}
 
-	if (!SrcMesh->RenderData->LODResources.IsValidIndex(LODIndex))
-	{
-		UE_LOG(LogMeshUtilities, Warning, TEXT("No mesh render data found for LOD%d %s."), LODIndex, *SrcMesh->GetName());
-		return false;
-	}
+void CullTrianglesFromVolumesAndUnderLandscapes(const UStaticMeshComponent* InMeshComponent, FRawMesh &OutRawMesh)
+{
+	UWorld* World = InMeshComponent->GetWorld();
+	TArray<ALandscapeProxy*> Landscapes;
+	TArray<AHLODMeshCullingVolume*> CullVolumes;
 
-	FStaticMeshSourceModel& SrcModel = SrcMesh->SourceModels[LODIndex];
-	const bool bImportedMesh = !SrcModel.RawMeshBulkData->IsEmpty();
-	const bool bReducedMesh = (SrcModel.ReductionSettings.PercentTriangles < 1.0f);
+	FBox ComponentBox(InMeshComponent->Bounds.Origin - InMeshComponent->Bounds.BoxExtent, InMeshComponent->Bounds.Origin + InMeshComponent->Bounds.BoxExtent);
 
-	if (bImportedMesh && !InMeshComponent->IsA<USplineMeshComponent>() && !bReducedMesh)
+	for (ULevel* Level : World->GetLevels())
 	{
-		SrcModel.RawMeshBulkData->LoadRawMesh(OutRawMesh);
-	}
-	else
-	{
-		// Reconstruct RawMesh from render data
-		ExportStaticMeshLOD(SrcMesh->RenderData->LODResources[LODIndex], OutRawMesh);
-	}
+		for (AActor* Actor : Level->Actors)
+		{
+			ALandscape* Proxy = Cast<ALandscape>(Actor);
+			if (Proxy && Proxy->bUseLandscapeForCullingInvisibleHLODVertices)
+			{
+				FVector Origin, Extent;
+				Proxy->GetActorBounds(false, Origin, Extent);
+				FBox LandscapeBox(Origin - Extent, Origin + Extent);
 
-	// Make sure the raw mesh is not irreparably malformed.
-	if (!OutRawMesh.IsValidOrFixable())
-	{
-		UE_LOG(LogMeshUtilities, Error, TEXT("Raw mesh (%s) is corrupt for LOD%d."), *SrcMesh->GetName(), LODIndex);
-		return false;
+				// Ignore Z axis for 2d bounds check
+				if (LandscapeBox.IntersectXY(ComponentBox))
+				{
+					Landscapes.Add(Proxy->GetLandscapeActor());
+				}
+			}
+
+			// Check for culling volumes
+			AHLODMeshCullingVolume* Volume = Cast<AHLODMeshCullingVolume>(Actor);
+			if (Volume)
+			{
+				// If the mesh's bounds intersect with the volume there is a possibility of culling
+				const bool bIntersecting = Volume->EncompassesPoint(InMeshComponent->Bounds.Origin, InMeshComponent->Bounds.SphereRadius, nullptr);
+				if (bIntersecting)
+				{
+					CullVolumes.Add(Volume);
+				}
+			}
+
+		}
 	}
+
+	TArray<bool> VertexVisible;
+	VertexVisible.AddZeroed(OutRawMesh.VertexPositions.Num());
+	int32 Index = 0;
 	
-	// Handle spline mesh deformation
-	if (InMeshComponent->IsA<USplineMeshComponent>())
+	for (const FVector& Position : OutRawMesh.VertexPositions)
 	{
-		USplineMeshComponent* SplineMeshComponent = Cast<USplineMeshComponent>(InMeshComponent);
+		// Start with setting visibility to true on all vertices
+		VertexVisible[Index] = true;
 
-		for (int32 iVert = 0; iVert < OutRawMesh.WedgeIndices.Num(); ++iVert)
+		// Check if this vertex is culled due to being underneath a landscape
+		if (Landscapes.Num() > 0)
 		{
-			uint32 Index = OutRawMesh.WedgeIndices[iVert];
-			float& AxisValue = USplineMeshComponent::GetAxisValue(OutRawMesh.VertexPositions[Index], SplineMeshComponent->ForwardAxis);
-			FTransform SliceTransform = SplineMeshComponent->CalcSliceTransform(AxisValue);
-			if (OutRawMesh.WedgeTangentX.Num())
+			bool bVertexWithinLandscapeBounds = false;
+
+			for (ALandscapeProxy* Proxy : Landscapes)
 			{
-				OutRawMesh.WedgeTangentX[iVert] = SliceTransform.TransformVector(OutRawMesh.WedgeTangentX[iVert]);
+				FVector Origin, Extent;
+				Proxy->GetActorBounds(false, Origin, Extent);
+				FBox LandscapeBox(Origin - Extent, Origin + Extent);
+				bVertexWithinLandscapeBounds |= LandscapeBox.IsInsideXY(Position);
 			}
 
-			if (OutRawMesh.WedgeTangentY.Num())
+			if (bVertexWithinLandscapeBounds)
 			{
-				OutRawMesh.WedgeTangentY[iVert] = SliceTransform.TransformVector(OutRawMesh.WedgeTangentY[iVert]);
+				const FVector Start = Position;
+				FVector End = Position - (WORLD_MAX * FVector::UpVector);
+				const bool IsAboveLandscape = IsLandscapeHit(Start, End, World, Landscapes);
+
+				End = Position + (WORLD_MAX * FVector::UpVector);
+				const bool IsUnderneathLandscape = IsLandscapeHit(Start, End, World, Landscapes);
+
+				// Vertex is visible when above landscape (with actual landscape underneath) or if there is no landscape beneath or above the vertex (falls outside of landscape bounds)
+				VertexVisible[Index] = (IsAboveLandscape && !IsUnderneathLandscape) || (!IsAboveLandscape && !IsUnderneathLandscape);
 			}
+		}
 
-			if (OutRawMesh.WedgeTangentZ.Num())
+		// Volume culling	
+		for (AHLODMeshCullingVolume* Volume : CullVolumes)
+		{
+			const bool bVertexIsInsideVolume = Volume->EncompassesPoint(Position, 0.0f, nullptr);
+			if (bVertexIsInsideVolume)
 			{
-				OutRawMesh.WedgeTangentZ[iVert] = SliceTransform.TransformVector(OutRawMesh.WedgeTangentZ[iVert]);
+				// Inside a culling volume so invisible
+				VertexVisible[Index] = false;
 			}
 		}
 
-		for (int32 iVert = 0; iVert < OutRawMesh.VertexPositions.Num(); ++iVert)
+		Index++;
+	}
+
+
+	// We now know which vertices are below the landscape
+	TArray<bool> TriangleVisible;
+	int32 NumTriangles = OutRawMesh.WedgeIndices.Num() / 3;
+	TriangleVisible.AddZeroed(NumTriangles);
+
+	bool bCreateNewMesh = false;
+
+	// Determine which triangles of the mesh are visible
+	for (int32 TriangleIndex = 0; TriangleIndex < NumTriangles; TriangleIndex++)
+	{
+		bool AboveLandscape = false;
+
+		for (int32 WedgeIndex = 0; WedgeIndex < 3; ++WedgeIndex)
+		{
+			AboveLandscape |= VertexVisible[OutRawMesh.WedgeIndices[(TriangleIndex * 3) + WedgeIndex]];
+		}
+		TriangleVisible[TriangleIndex] = AboveLandscape;
+		bCreateNewMesh |= !AboveLandscape;
+
+		for (int32 WedgeIndex = 0; WedgeIndex < 3; ++WedgeIndex)
 		{
-			float& AxisValue = USplineMeshComponent::GetAxisValue(OutRawMesh.VertexPositions[iVert], SplineMeshComponent->ForwardAxis);
-			FTransform SliceTransform = SplineMeshComponent->CalcSliceTransform(AxisValue);
-			AxisValue = 0.0f;
-			OutRawMesh.VertexPositions[iVert] = SliceTransform.TransformPosition(OutRawMesh.VertexPositions[iVert]);
+			AboveLandscape |= VertexVisible[OutRawMesh.WedgeIndices[(TriangleIndex * 3) + WedgeIndex]];
 		}
+		TriangleVisible[TriangleIndex] = AboveLandscape;
+		bCreateNewMesh |= !AboveLandscape;
+		
 	}
 
-	// Use build settings from base mesh for LOD entries that was generated inside Editor.
-	const FMeshBuildSettings& BuildSettings = bImportedMesh ? SrcModel.BuildSettings : SrcMesh->SourceModels[0].BuildSettings;
+	// Check whether or not we have to create a new mesh
+	if (bCreateNewMesh)
+	{
+		FRawMesh NewRawMesh;
+		TMap<int32, int32> VertexRemapping;
 
-	// Transform raw mesh to world space
-	FTransform CtoM = InMeshComponent->ComponentToWorld;
-	// Take into account build scale settings only for meshes imported from raw data
-	// meshes reconstructed from render data already have build scale applied
-	if (bImportedMesh)
+		// Fill new mesh with data only from visible triangles
+		for (int32 TriangleIndex = 0; TriangleIndex < NumTriangles; ++TriangleIndex)
+		{
+			if (!TriangleVisible[TriangleIndex])
+				continue;
+
+			for (int32 WedgeIndex = 0; WedgeIndex < 3; ++WedgeIndex)
+			{
+				int32 OldIndex = OutRawMesh.WedgeIndices[(TriangleIndex * 3) + WedgeIndex];
+
+				int32 NewIndex;
+
+				int32* RemappedIndex = VertexRemapping.Find(Index);
+				if (RemappedIndex)
+				{
+					NewIndex = *RemappedIndex;
+				}
+				else
+				{
+					NewIndex = NewRawMesh.VertexPositions.Add(OutRawMesh.VertexPositions[OldIndex]);
+					VertexRemapping.Add(OldIndex, NewIndex);
+				}
+
+				NewRawMesh.WedgeIndices.Add(NewIndex);
+				if (OutRawMesh.WedgeColors.Num()) NewRawMesh.WedgeColors.Add(OutRawMesh.WedgeColors[(TriangleIndex * 3) + WedgeIndex]);
+				if (OutRawMesh.WedgeTangentX.Num()) NewRawMesh.WedgeTangentX.Add(OutRawMesh.WedgeTangentX[(TriangleIndex * 3) + WedgeIndex]);
+				if (OutRawMesh.WedgeTangentY.Num()) NewRawMesh.WedgeTangentY.Add(OutRawMesh.WedgeTangentY[(TriangleIndex * 3) + WedgeIndex]);
+				if (OutRawMesh.WedgeTangentZ.Num()) NewRawMesh.WedgeTangentZ.Add(OutRawMesh.WedgeTangentZ[(TriangleIndex * 3) + WedgeIndex]);
+
+				for (int32 UVIndex = 0; UVIndex < MAX_MESH_TEXTURE_COORDS; ++UVIndex)
+				{
+					if (OutRawMesh.WedgeTexCoords[UVIndex].Num())
+					{
+						NewRawMesh.WedgeTexCoords[UVIndex].Add(OutRawMesh.WedgeTexCoords[UVIndex][(TriangleIndex * 3) + WedgeIndex]);
+					}
+				}
+			}
+
+			NewRawMesh.FaceMaterialIndices.Add(OutRawMesh.FaceMaterialIndices[TriangleIndex]);
+			NewRawMesh.FaceSmoothingMasks.Add(OutRawMesh.FaceSmoothingMasks[TriangleIndex]);
+		}
+
+		OutRawMesh = NewRawMesh;
+	}
+
+}
+
+
+void PropagateSplineDeformationToRawMesh(const USplineMeshComponent* InSplineMeshComponent, FRawMesh &OutRawMesh) 
+{
+	// Apply spline deformation for each vertex's tangents
+	for (int32 iVert = 0; iVert < OutRawMesh.WedgeIndices.Num(); ++iVert)
 	{
-		CtoM.SetScale3D(CtoM.GetScale3D()*BuildSettings.BuildScale3D);
+		uint32 Index = OutRawMesh.WedgeIndices[iVert];
+		float& AxisValue = USplineMeshComponent::GetAxisValue(OutRawMesh.VertexPositions[Index], InSplineMeshComponent->ForwardAxis);
+		FTransform SliceTransform = InSplineMeshComponent->CalcSliceTransform(AxisValue);
+
+		// Transform tangents first
+		if (OutRawMesh.WedgeTangentX.Num())
+		{
+			OutRawMesh.WedgeTangentX[iVert] = SliceTransform.TransformVector(OutRawMesh.WedgeTangentX[iVert]);
+		}
+
+		if (OutRawMesh.WedgeTangentY.Num())
+		{
+			OutRawMesh.WedgeTangentY[iVert] = SliceTransform.TransformVector(OutRawMesh.WedgeTangentY[iVert]);
+		}
+
+		if (OutRawMesh.WedgeTangentZ.Num())
+		{
+			OutRawMesh.WedgeTangentZ[iVert] = SliceTransform.TransformVector(OutRawMesh.WedgeTangentZ[iVert]);
+		}
+	}
+
+	// Apply spline deformation for each vertex position
+	for (int32 iVert = 0; iVert < OutRawMesh.VertexPositions.Num(); ++iVert)
+	{
+		float& AxisValue = USplineMeshComponent::GetAxisValue(OutRawMesh.VertexPositions[iVert], InSplineMeshComponent->ForwardAxis);
+		FTransform SliceTransform = InSplineMeshComponent->CalcSliceTransform(AxisValue);
+		AxisValue = 0.0f;
+		OutRawMesh.VertexPositions[iVert] = SliceTransform.TransformPosition(OutRawMesh.VertexPositions[iVert]);
 	}
+}
 
+
+void TransformRawMeshVertexData(const FTransform& InTransform, FRawMesh &OutRawMesh )
+{
 	for (FVector& Vertex : OutRawMesh.VertexPositions)
 	{
-		Vertex = CtoM.TransformPosition(Vertex);
+		Vertex = InTransform.TransformPosition(Vertex);
 	}
+
 	for (FVector& TangentX : OutRawMesh.WedgeTangentX)
 	{
-		TangentX = CtoM.TransformVectorNoScale(TangentX);
+		TangentX = InTransform.TransformVectorNoScale(TangentX);
 	}
+
 	for (FVector& TangentY : OutRawMesh.WedgeTangentY)
 	{
-		TangentY = CtoM.TransformVectorNoScale(TangentY);
+		TangentY = InTransform.TransformVectorNoScale(TangentY);
 	}
+
 	for (FVector& TangentZ : OutRawMesh.WedgeTangentZ)
 	{
-		TangentZ = CtoM.TransformVectorNoScale(TangentZ);
+		TangentZ = InTransform.TransformVectorNoScale(TangentZ);
 	}
 	
-	const bool bIsMirrored = CtoM.GetDeterminant() < 0.f;
+	const bool bIsMirrored = InTransform.GetDeterminant() < 0.f;
 	if (bIsMirrored)
 	{
 		// Flip faces
@@ -5753,12 +5934,12 @@ bool FMeshUtilities::ConstructRawMesh(
 			}
 		}
 	}
+}
 
-	int32 NumWedges = OutRawMesh.WedgeIndices.Num();
 
-	// Figure out if we should recompute normals and tangents. By default generated LODs should not recompute normals
-	bool bRecomputeNormals = (bImportedMesh && BuildSettings.bRecomputeNormals) || OutRawMesh.WedgeTangentZ.Num() == 0 || bIsMirrored;
-	bool bRecomputeTangents = (bImportedMesh && BuildSettings.bRecomputeTangents) || OutRawMesh.WedgeTangentX.Num() == 0 || OutRawMesh.WedgeTangentY.Num() == 0 || bIsMirrored;
+void RecomputeTangentsAndNormalsForRawMesh(bool bRecomputeTangents, bool bRecomputeNormals, const FMeshBuildSettings& InBuildSettings, FRawMesh &OutRawMesh )
+{
+	const int32 NumWedges = OutRawMesh.WedgeIndices.Num();
 
 	// Dump normals and tangents if we are recomputing them.
 	if (bRecomputeTangents)
@@ -5778,18 +5959,18 @@ bool FMeshUtilities::ConstructRawMesh(
 	// Compute any missing tangents.
 	if (bRecomputeNormals || bRecomputeTangents)
 	{
-		float ComparisonThreshold = GetComparisonThreshold(BuildSettings);
+		float ComparisonThreshold = GetComparisonThreshold(InBuildSettings);
 		TMultiMap<int32, int32> OverlappingCorners;
 		FindOverlappingCorners(OverlappingCorners, OutRawMesh, ComparisonThreshold);
 
 		// Static meshes always blend normals of overlapping corners.
 		uint32 TangentOptions = ETangentOptions::BlendOverlappingNormals;
-		if (BuildSettings.bRemoveDegenerates)
+		if (InBuildSettings.bRemoveDegenerates)
 		{
 			// If removing degenerate triangles, ignore them when computing tangents.
 			TangentOptions |= ETangentOptions::IgnoreDegenerateTriangles;
 		}
-		if (BuildSettings.bUseMikkTSpace)
+		if (InBuildSettings.bUseMikkTSpace)
 		{
 			ComputeTangents_MikkTSpace(OutRawMesh, OverlappingCorners, TangentOptions);
 		}
@@ -5803,12 +5984,116 @@ bool FMeshUtilities::ConstructRawMesh(
 	check(OutRawMesh.WedgeTangentX.Num() == NumWedges);
 	check(OutRawMesh.WedgeTangentY.Num() == NumWedges);
 	check(OutRawMesh.WedgeTangentZ.Num() == NumWedges);
+}
+
+bool FMeshUtilities::ConstructRawMesh(
+	const UStaticMeshComponent* InMeshComponent,
+	int32 InLODIndex,
+	const bool bPropagateVertexColours,
+	FRawMesh& OutRawMesh,
+	TArray<UMaterialInterface*>& OutUniqueMaterials,
+	TArray<int32>& OutGlobalMaterialIndices) const
+{
+	// Retrieve source static mesh 
+	const UStaticMesh* SourceStaticMesh = InMeshComponent->StaticMesh;
+	
+	if (SourceStaticMesh == NULL)
+	{
+		UE_LOG(LogMeshUtilities, Warning, TEXT("No static mesh actor found in component %s."), *InMeshComponent->GetName());
+		return false;
+	}
+
+	if (!SourceStaticMesh->SourceModels.IsValidIndex(InLODIndex))
+	{
+		UE_LOG(LogMeshUtilities, Log, TEXT("No mesh data found for LOD%d %s."), InLODIndex, *SourceStaticMesh->GetName());
+		return false;
+	}
+
+	if (!SourceStaticMesh->RenderData->LODResources.IsValidIndex(InLODIndex))
+	{
+		UE_LOG(LogMeshUtilities, Warning, TEXT("No mesh render data found for LOD%d %s."), InLODIndex, *SourceStaticMesh->GetName());
+		return false;
+	}
+
+	const FStaticMeshSourceModel& SourceStaticMeshModel = SourceStaticMesh->SourceModels[InLODIndex];
 
+	// Imported meshes will have a filled RawMeshBulkData set
+	const bool bImportedMesh = !SourceStaticMeshModel.RawMeshBulkData->IsEmpty();
+	// Check whether or not this mesh has been reduced in-engine
+	const bool bReducedMesh = (SourceStaticMeshModel.ReductionSettings.PercentTriangles < 1.0f);
+	// rying to retrieve rawmesh from SourceStaticMeshModel was giving issues, which causes a mismatch
+	const bool bRenderDataMismatch = (InLODIndex > 0);
+	
+	// Determine whether we load the raw mesh data from (original) import data or from the generated render data resources
+	if (bImportedMesh && !InMeshComponent->IsA<USplineMeshComponent>() && !bReducedMesh && !bRenderDataMismatch)
+	{
+		SourceStaticMeshModel.RawMeshBulkData->LoadRawMesh(OutRawMesh);
+	}
+	else
+	{		
+		ExportStaticMeshLOD(SourceStaticMesh->RenderData->LODResources[InLODIndex], OutRawMesh);
+	}
+
+	// Make sure the raw mesh is not irreparably malformed.
+	if (!OutRawMesh.IsValidOrFixable())
+	{
+		UE_LOG(LogMeshUtilities, Error, TEXT("Raw mesh (%s) is corrupt for LOD%d."), *SourceStaticMesh->GetName(), InLODIndex);
+		return false;
+	}
+	
+	// Handle spline mesh deformation
+	if (InMeshComponent->IsA<USplineMeshComponent>())
+	{
+		const USplineMeshComponent* SplineMeshComponent = Cast<USplineMeshComponent>(InMeshComponent);
+		// Deform raw mesh data according to the Spline Mesh Component's data
+		PropagateSplineDeformationToRawMesh(SplineMeshComponent, OutRawMesh);
+	}
+
+	// Use build settings from base mesh for LOD entries that was generated inside Editor.
+	const FMeshBuildSettings& BuildSettings = bImportedMesh ? SourceStaticMeshModel.BuildSettings : SourceStaticMesh->SourceModels[0].BuildSettings;
+
+	// Transform raw mesh to world space
+	FTransform ComponentToWorldTransform = InMeshComponent->ComponentToWorld;
+	// Take into account build scale settings only for meshes imported from raw data
+	// meshes reconstructed from render data already have build scale applied
+	if (bImportedMesh)
+	{
+		ComponentToWorldTransform.SetScale3D(ComponentToWorldTransform.GetScale3D()*BuildSettings.BuildScale3D);
+	}
+
+	// Transform raw mesh vertex data by the Static Mesh Component's component to world transformation	
+	TransformRawMeshVertexData(ComponentToWorldTransform, OutRawMesh);
+	
+	// If specified propagate painted vertex colors into our raw mesh
+	if (bPropagateVertexColours)
+	{
+		PropagatePaintedColorsToRawMesh(InMeshComponent, InLODIndex, OutRawMesh);
+	}
+
+	// Remove all triangles underneath terrain
+	CullTrianglesFromVolumesAndUnderLandscapes(InMeshComponent, OutRawMesh);
+	
+	// Culling triangles could lead to an entirely empty RawMesh (all vertices culled)
+	if (!OutRawMesh.IsValid())
+	{
+		return false;
+	}
+
+	// Figure out if we should recompute normals and tangents. By default generated LODs should not recompute normals
+	const bool bIsMirrored = ComponentToWorldTransform.GetDeterminant() < 0.f;
+	bool bRecomputeNormals = (bImportedMesh && BuildSettings.bRecomputeNormals) || OutRawMesh.WedgeTangentZ.Num() == 0 || bIsMirrored;
+	bool bRecomputeTangents = (bImportedMesh && BuildSettings.bRecomputeTangents) || OutRawMesh.WedgeTangentX.Num() == 0 || OutRawMesh.WedgeTangentY.Num() == 0 || bIsMirrored;
+
+	if (bRecomputeNormals || bRecomputeTangents)
+	{
+		RecomputeTangentsAndNormalsForRawMesh(bRecomputeTangents, bRecomputeNormals, BuildSettings, OutRawMesh);
+	}
+
+	// Retrieving materials
 	UMaterialInterface* DefaultMaterial = Cast<UMaterialInterface>(UMaterial::GetDefaultMaterial(MD_Surface));
 
-	//Need to store the unique material indices in order to re-map the material indices in each rawmesh
-	//Only using the base mesh
-	for (const FStaticMeshSection& Section : SrcMesh->RenderData->LODResources[LODIndex].Sections)
+	//Need to store the unique material indices in order to re-map the material indices in each rawmesh	
+	for (const FStaticMeshSection& Section : SourceStaticMesh->RenderData->LODResources[InLODIndex].Sections)
 	{
 		// Add material and store the material ID
 		UMaterialInterface* MaterialToAdd = InMeshComponent->GetMaterial(Section.MaterialIndex);
@@ -5828,7 +6113,19 @@ bool FMeshUtilities::ConstructRawMesh(
 		}
 
 		const int32 MaterialIdx = OutUniqueMaterials.Add(MaterialToAdd);
-		OutGlobalMaterialIndices.Add(MaterialIdx);
+		const int32 MaterialMapIdx = OutGlobalMaterialIndices.Add(MaterialIdx);
+		
+		// Update face material indices?
+		if (OutRawMesh.FaceMaterialIndices.Num())
+		{
+			for (int32& MaterialIndex : OutRawMesh.FaceMaterialIndices)
+			{
+				if (MaterialIndex == Section.MaterialIndex)
+				{
+					MaterialIndex = MaterialMapIdx;
+				}
+			}
+		}
 	}
 
 	return true;
@@ -5911,7 +6208,7 @@ void FMeshUtilities::ExtractMeshDataForGeometryCache(FRawMesh& RawMesh, const FM
 /*------------------------------------------------------------------------------
 Mesh merging
 ------------------------------------------------------------------------------*/
-bool FMeshUtilities::PropagatePaintedColorsToRawMesh(UStaticMeshComponent* StaticMeshComponent, int32 LODIndex, FRawMesh& RawMesh) const
+bool FMeshUtilities::PropagatePaintedColorsToRawMesh(const UStaticMeshComponent* StaticMeshComponent, int32 LODIndex, FRawMesh& RawMesh) const
 {
 	UStaticMesh* StaticMesh = StaticMeshComponent->StaticMesh;
 
@@ -5925,10 +6222,11 @@ bool FMeshUtilities::PropagatePaintedColorsToRawMesh(UStaticMeshComponent* Stati
 		FStaticMeshLODResources& RenderModel = RenderData.LODResources[LODIndex];
 
 		if (ColorVertexBuffer.GetNumVertices() == RenderModel.GetNumVertices())
-		{
+		{	
 			int32 NumWedges = RawMesh.WedgeIndices.Num();
+			const bool bUseWedgeMap = RenderData.WedgeMap.Num() > 0 && RenderData.WedgeMap.Num() == NumWedges;
 			// If we have a wedge map
-			if (RenderData.WedgeMap.Num() > 0)
+			if (bUseWedgeMap)
 			{
 				if (RenderData.WedgeMap.Num() == NumWedges)
 				{
@@ -5952,29 +6250,31 @@ bool FMeshUtilities::PropagatePaintedColorsToRawMesh(UStaticMeshComponent* Stati
 
 					return true;
 				}
-				else
-				{
-					UE_LOG(LogMeshUtilities, Warning, TEXT("{%s} Wedge map size %d is wrong. Expected %d."), *StaticMesh->GetName(), RenderData.WedgeMap.Num(), RawMesh.WedgeIndices.Num());
-				}
 			}
 			// No wedge map (this can happen when we poly reduce the LOD for example)
 			// Use index buffer directly
 			else 
 			{
-				RawMesh.WedgeColors.SetNumUninitialized(NumWedges);
+				UE_LOG(LogMeshUtilities, Warning, TEXT("{%s} Wedge map size %d is wrong or empty. Expected %d. Falling back on using index buffer for propagating vertex painting"), *StaticMesh->GetName(), RenderData.WedgeMap.Num(), RawMesh.WedgeIndices.Num());
 
-				for (int32 i = 0; i < NumWedges; ++i)
+				RawMesh.WedgeColors.SetNumUninitialized(NumWedges);
+				
+				if (RawMesh.VertexPositions.Num() == ColorVertexBuffer.GetNumVertices())
 				{
-					FColor WedgeColor = FColor::White;
-					uint32 VertIndex = RawMesh.WedgeIndices[i];
-					if(VertIndex < ColorVertexBuffer.GetNumVertices())
+					for (int32 i = 0; i < NumWedges; ++i)
 					{
-						WedgeColor = ColorVertexBuffer.VertexColor(VertIndex);
+						FColor WedgeColor = FColor::White;
+						uint32 VertIndex = RawMesh.WedgeIndices[i]; 
+						
+						if (VertIndex < ColorVertexBuffer.GetNumVertices())
+						{
+							WedgeColor = ColorVertexBuffer.VertexColor(VertIndex);
+						}
+						RawMesh.WedgeColors[i] = WedgeColor;
 					}
-					RawMesh.WedgeColors[i] = WedgeColor;
-				}
 
-				return true;
+					return true;
+				}
 			}
 		}	
 	}
@@ -6620,8 +6920,8 @@ void FMeshUtilities::MergeStaticMeshComponents(const TArray<UStaticMeshComponent
 	MainTask.EnterProgressFrame(10, LOCTEXT("MeshUtilities_MergeStaticMeshComponents_RetrievingRawMesh", "Retrieving Raw Meshes"));
 
 	// Make sure if we are merging materials, that we will only export LOD0 (currently no support for baking out multiple textures and materials per LOD)
-	StartLODIndex = (InSettings.bMergeMaterials) ? 0 : StartLODIndex;
-	NumMaxLOD = (InSettings.bMergeMaterials) ? 1 : NumMaxLOD;
+	StartLODIndex = 0;// (InSettings.bMergeMaterials) ? 0 : StartLODIndex;
+	NumMaxLOD = 1; // (InSettings.bMergeMaterials) ? 1 : NumMaxLOD;
 
 	int32 RawMeshLODIdx = 0;
 	for (int32 LODIndex = StartLODIndex; LODIndex < NumMaxLOD; ++LODIndex, ++RawMeshLODIdx)
@@ -6645,15 +6945,13 @@ void FMeshUtilities::MergeStaticMeshComponents(const TArray<UStaticMeshComponent
 
 			FRawMesh& RawMeshLOD = SourceMeshes[MeshId].MeshLODData[LODIndex].RawMesh;
 
-			if (ConstructRawMesh(StaticMeshComponent, ExportLODIndex, RawMeshLOD, UniqueMaterials, MeshMaterialMap))
+			if (ConstructRawMesh(StaticMeshComponent, ExportLODIndex, InSettings.bBakeVertexData, RawMeshLOD, UniqueMaterials, MeshMaterialMap))
 			{
 				MaterialMap.Add(FMeshIdAndLOD(MeshId, RawMeshLODIdx), MeshMaterialMap);
 
-				// Should we use vertex colors?
+				// Check if vertex colours are propagated
 				if (InSettings.bBakeVertexData)
 				{
-					// Propagate painted vertex colors into our raw mesh
-					PropagatePaintedColorsToRawMesh(StaticMeshComponent, ExportLODIndex, RawMeshLOD);
 					// Whether at least one of the meshes has vertex colors
 					bWithVertexColors[RawMeshLODIdx] |= (RawMeshLOD.WedgeColors.Num() != 0);
 				}
@@ -6687,22 +6985,7 @@ void FMeshUtilities::MergeStaticMeshComponents(const TArray<UStaticMeshComponent
 	}
 #endif
 
-	TArray<bool> MeshShouldBakeVertexData;
-	TMap<FMeshIdAndLOD, TArray<int32> > NewMaterialMap;
-	TArray<UMaterialInterface*> NewStaticMeshMaterials;
-	FMaterialUtilities::RemapUniqueMaterialIndices(
-		UniqueMaterials,
-		SourceMeshes,
-		MaterialMap,
-		InSettings.MaterialSettings,
-		InSettings.bBakeVertexData,
-		InSettings.bMergeMaterials,
-		MeshShouldBakeVertexData,
-		NewMaterialMap,
-		NewStaticMeshMaterials);
-	// Use shared material data.
-	Exchange(MaterialMap, NewMaterialMap);
-	Exchange(UniqueMaterials, NewStaticMeshMaterials);
+	
 
 	MainTask.EnterProgressFrame(20);
 
@@ -6711,6 +6994,24 @@ void FMeshUtilities::MergeStaticMeshComponents(const TArray<UStaticMeshComponent
 
 	if (InSettings.bMergeMaterials)
 	{
+		TArray<bool> MeshShouldBakeVertexData;
+		TMap<FMeshIdAndLOD, TArray<int32> > NewMaterialMap;
+		TArray<UMaterialInterface*> NewStaticMeshMaterials;
+		FMaterialUtilities::RemapUniqueMaterialIndices(
+			UniqueMaterials,
+			SourceMeshes,
+			MaterialMap,
+			InSettings.MaterialSettings,
+			InSettings.bBakeVertexData,
+			InSettings.bMergeMaterials,
+			MeshShouldBakeVertexData,
+			NewMaterialMap,
+			NewStaticMeshMaterials);
+		// Use shared material data.
+		Exchange(MaterialMap, NewMaterialMap);
+		Exchange(UniqueMaterials, NewStaticMeshMaterials);
+
+
 		// Should merge flattened materials into one texture
 		MainTask.EnterProgressFrame(20, LOCTEXT("MeshUtilities_MergeStaticMeshComponents_MergingMaterials", "Merging Materials"));
 
diff --git a/Engine/Source/Developer/MeshUtilities/Public/MeshMergeData.h b/Engine/Source/Developer/MeshUtilities/Public/MeshMergeData.h
--- a/Engine/Source/Developer/MeshUtilities/Public/MeshMergeData.h
+++ b/Engine/Source/Developer/MeshUtilities/Public/MeshMergeData.h
@@ -7,6 +7,9 @@
 /** Structure holding intermediate mesh merging data that is used throughout the mesh merging and proxy creation processes */
 struct FMeshMergeData
 {
+	FMeshMergeData()
+		: SourceStaticMesh(nullptr)
+	{}
 	/** Raw mesh data from the source static mesh */
 	FRawMesh RawMesh;	
 	/** Contains the original texture bounds, if the material requires baking down per-vertex data */
diff --git a/Engine/Source/Developer/MeshUtilities/Public/MeshUtilities.h b/Engine/Source/Developer/MeshUtilities/Public/MeshUtilities.h
--- a/Engine/Source/Developer/MeshUtilities/Public/MeshUtilities.h
+++ b/Engine/Source/Developer/MeshUtilities/Public/MeshUtilities.h
@@ -265,7 +265,7 @@ public:
 	* @param InProxyCreatedDelegate - Delegate callback for when the proxy is finished
 	* @param bAllowAsync - Flag whether or not this call could be run async (SimplygonSwarm)
 	*/
-	virtual void CreateProxyMesh(const TArray<class AActor*>& InActors, const struct FMeshProxySettings& InMeshProxySettings, UPackage* InOuter, const FString& InProxyBasePackageName, const FGuid InGuid, FCreateProxyDelegate InProxyCreatedDelegate, const bool bAllowAsync = false) = 0;
+	virtual void CreateProxyMesh(const TArray<class AActor*>& InActors, const struct FMeshProxySettings& InMeshProxySettings, UPackage* InOuter, const FString& InProxyBasePackageName, const FGuid InGuid, FCreateProxyDelegate InProxyCreatedDelegate, const bool bAllowAsync = false, const float ScreenAreaSize = 1.0f) = 0;
 
 	/**
 	*	Merges list of actors into single proxy mesh
@@ -300,7 +300,8 @@ public:
 		const struct FMeshProxySettings& InProxySettings,
 		UPackage* InOuter,
 		const FString& ProxyBasePackageName,
-		TArray<UObject*>& OutAssetsToSync ) = 0;
+		TArray<UObject*>& OutAssetsToSync, 		
+		const float ScreenAreaSize = 1.0f) = 0;
 
 	/**
 	* FlattenMaterialsWithMeshData
@@ -332,7 +333,7 @@ public:
 	* @param RawMesh - Target raw mesh to propagate vertex colours to
 	* @return bool - whether or not propagating succeeded
 	*/
-	virtual bool PropagatePaintedColorsToRawMesh(UStaticMeshComponent* StaticMeshComponent, int32 LODIndex, FRawMesh& RawMesh) const = 0;
+	virtual bool PropagatePaintedColorsToRawMesh(const UStaticMeshComponent* StaticMeshComponent, int32 LODIndex, FRawMesh& RawMesh) const = 0;
 	
 	/**
 	*	Calculates UV coordinates bounds for the given Raw Mesh
