Deprecation fix for AddThirdPartyPrivateStaticDependencies

diff --git a/Engine/Source/Developer/MeshUtilities/MeshUtilities.Build.cs b/Engine/Source/Developer/MeshUtilities/MeshUtilities.Build.cs
--- a/Engine/Source/Developer/MeshUtilities/MeshUtilities.Build.cs
+++ b/Engine/Source/Developer/MeshUtilities/MeshUtilities.Build.cs
@@ -31,19 +31,19 @@ public class MeshUtilities : ModuleRules
 			}
 		);
 
-        AddThirdPartyPrivateStaticDependencies(Target, "nvTriStrip");
-		AddThirdPartyPrivateStaticDependencies(Target, "ForsythTriOptimizer");
-		AddThirdPartyPrivateStaticDependencies(Target, "MeshSimplifier");
-		AddThirdPartyPrivateStaticDependencies(Target, "MikkTSpace");
+        AddEngineThirdPartyPrivateStaticDependencies(Target, "nvTriStrip");
+        AddEngineThirdPartyPrivateStaticDependencies(Target, "ForsythTriOptimizer");
+        AddEngineThirdPartyPrivateStaticDependencies(Target, "MeshSimplifier");
+        AddEngineThirdPartyPrivateStaticDependencies(Target, "MikkTSpace");
 
 		if ((Target.Platform == UnrealTargetPlatform.Win64) || (Target.Platform == UnrealTargetPlatform.Win32))
 		{
-			AddThirdPartyPrivateStaticDependencies(Target, "DX9");
+            AddEngineThirdPartyPrivateStaticDependencies(Target, "DX9");
 		}
 
 		if (UEBuildConfiguration.bCompileSimplygon == true)
 		{
-			AddThirdPartyPrivateDynamicDependencies(Target, "SimplygonMeshReduction");
+            AddEngineThirdPartyPrivateStaticDependencies(Target, "SimplygonMeshReduction");
             
             if (UEBuildConfiguration.bCompileSimplygonSSF == true)
             {
