Static -> Dynamic include

diff --git a/Engine/Source/Developer/MeshUtilities/MeshUtilities.Build.cs b/Engine/Source/Developer/MeshUtilities/MeshUtilities.Build.cs
--- a/Engine/Source/Developer/MeshUtilities/MeshUtilities.Build.cs
+++ b/Engine/Source/Developer/MeshUtilities/MeshUtilities.Build.cs
@@ -43,7 +43,7 @@ public class MeshUtilities : ModuleRules
 
 		if (UEBuildConfiguration.bCompileSimplygon == true)
 		{
-            AddEngineThirdPartyPrivateStaticDependencies(Target, "SimplygonMeshReduction");
+            AddEngineThirdPartyPrivateDynamicDependencies(Target, "SimplygonMeshReduction");
             
             if (UEBuildConfiguration.bCompileSimplygonSSF == true)
             {
