Copying //UE4/Dev-Framework to //UE4/Dev-Main (Source: //UE4/Dev-Framework @ 3136612)

diff --git a/Engine/Source/Developer/MeshUtilities/Private/MeshUtilities.cpp b/Engine/Source/Developer/MeshUtilities/Private/MeshUtilities.cpp
--- a/Engine/Source/Developer/MeshUtilities/Private/MeshUtilities.cpp
+++ b/Engine/Source/Developer/MeshUtilities/Private/MeshUtilities.cpp
@@ -30,6 +30,19 @@
 #include "LandscapeHeightfieldCollisionComponent.h"
 #include "Engine/HLODMeshCullingVolume.h"
 #include "ProxyMaterialUtilities.h"
+#include "LevelEditor.h"
+#include "IAnimationBlueprintEditorModule.h"
+#include "IAnimationEditorModule.h"
+#include "ISkeletalMeshEditorModule.h"
+#include "ISkeletonEditorModule.h"
+#include "IPersonaToolkit.h"
+#include "Dialogs/DlgPickAssetPath.h"
+#include "SkeletalRenderPublic.h"
+#include "AssetRegistryModule.h"
+#include "SNotificationList.h"
+#include "NotificationManager.h"
+#include "Toolkits/AssetEditorManager.h"
+#include "StaticMeshResources.h"
 
 //@todo - implement required vector intrinsics for other implementations
 #if PLATFORM_ENABLE_VECTORINTRINSICS
@@ -157,6 +170,15 @@ private:
 
 	virtual void CalcBoneVertInfos(USkeletalMesh* SkeletalMesh, TArray<FBoneVertInfo>& Infos, bool bOnlyDominant) override;
 
+	/** 
+	 * Convert a set of mesh components in their current pose to a static mesh. 
+	 * @param	InMeshComponents		The mesh components we want to convert
+	 * @param	InRootTransform			The transform of the root of the mesh we want to output
+	 * @param	InPackageName			The package name to create the static mesh in. If this is empty then a dialog will be displayed to pick the mesh.
+	 * @return a new static mesh (specified by the user)
+	 */
+	virtual UStaticMesh* ConvertMeshesToStaticMesh(const TArray<UMeshComponent*>& InMeshComponents, const FTransform& InRootTransform = FTransform::Identity, const FString& InPackageName = FString()) override;
+
 	/**
 	* Builds a renderable skeletal mesh LOD model. Note that the array of chunks
 	* will be destroyed during this process!
@@ -265,6 +287,48 @@ private:
 
 	// Need to call some members from this class, (which is internal to this module)
 	friend class FStaticMeshUtilityBuilder;
+
+protected:
+	void AddAnimationBlueprintEditorToolbarExtender();
+
+	void RemoveAnimationBlueprintEditorToolbarExtender();
+
+	TSharedRef<FExtender> GetAnimationBlueprintEditorToolbarExtender(const TSharedRef<FUICommandList> CommandList, TSharedRef<IAnimationBlueprintEditor> InAnimationBlueprintEditor);
+
+	void AddAnimationEditorToolbarExtender();
+
+	void RemoveAnimationEditorToolbarExtender();
+
+	TSharedRef<FExtender> GetAnimationEditorToolbarExtender(const TSharedRef<FUICommandList> CommandList, TSharedRef<IAnimationEditor> InAnimationEditor);
+
+	void AddSkeletalMeshEditorToolbarExtender();
+
+	void RemoveSkeletalMeshEditorToolbarExtender();
+
+	TSharedRef<FExtender> GetSkeletalMeshEditorToolbarExtender(const TSharedRef<FUICommandList> CommandList, TSharedRef<ISkeletalMeshEditor> InSkeletalMeshEditor);
+
+	void AddSkeletonEditorToolbarExtender();
+
+	void RemoveSkeletonEditorToolbarExtender();
+
+	TSharedRef<FExtender> GetSkeletonEditorToolbarExtender(const TSharedRef<FUICommandList> CommandList, TSharedRef<ISkeletonEditor> InSkeletonEditor);
+
+	void HandleAddConvertComponentToStaticMeshToToolbar(FToolBarBuilder& ParentToolbarBuilder, UMeshComponent* MeshComponent);
+
+	void AddLevelViewportMenuExtender();
+
+	void RemoveLevelViewportMenuExtender();
+
+	TSharedRef<FExtender> GetLevelViewportContextMenuExtender(const TSharedRef<FUICommandList> CommandList, const TArray<AActor*> InActors);
+
+	void ConvertActorMeshesToStaticMesh(const TArray<AActor*> InActors);
+
+	FDelegateHandle ModuleLoadedDelegateHandle;
+	FDelegateHandle LevelViewportExtenderHandle;
+	FDelegateHandle AnimationBlueprintEditorExtenderHandle;
+	FDelegateHandle AnimationEditorExtenderHandle;
+	FDelegateHandle SkeletalMeshEditorExtenderHandle;
+	FDelegateHandle SkeletonEditorExtenderHandle;
 };
 
 IMPLEMENT_MODULE(FMeshUtilities, MeshUtilities);
@@ -1262,6 +1326,433 @@ void FMeshUtilities::CalcBoneVertInfos(USkeletalMesh* SkeletalMesh, TArray<FBone
 	SkeletalMeshTools::CalcBoneVertInfos(SkeletalMesh, Infos, bOnlyDominant);
 }
 
+// Helper function for ConvertMeshesToStaticMesh
+static void AddOrDuplicateMaterial(UMaterialInterface* InMaterialInterface, const FString& InPackageName, TArray<UMaterialInterface*>& OutMaterials)
+{
+	if (InMaterialInterface && !InMaterialInterface->GetOuter()->IsA<UPackage>())
+	{
+		// Convert runtime material instances to new concrete material instances
+		// Create new package
+		FString OriginalMaterialName = InMaterialInterface->GetName();
+		FString MaterialPath = FPackageName::GetLongPackagePath(InPackageName) / OriginalMaterialName;
+		FString MaterialName;
+		FAssetToolsModule& AssetToolsModule = FModuleManager::LoadModuleChecked<FAssetToolsModule>("AssetTools");
+		AssetToolsModule.Get().CreateUniqueAssetName(MaterialPath, TEXT(""), MaterialPath, MaterialName);
+		UPackage* MaterialPackage = CreatePackage(NULL, *MaterialPath);
+
+		// Duplicate the object into the new package
+		UMaterialInterface* NewMaterialInterface = DuplicateObject<UMaterialInterface>(InMaterialInterface, MaterialPackage, *MaterialName);
+		NewMaterialInterface->SetFlags(RF_Public | RF_Standalone);
+
+		if (UMaterialInstanceDynamic* MaterialInstanceDynamic = Cast<UMaterialInstanceDynamic>(NewMaterialInterface))
+		{
+			UMaterialInstanceDynamic* OldMaterialInstanceDynamic = CastChecked<UMaterialInstanceDynamic>(InMaterialInterface);
+			MaterialInstanceDynamic->K2_CopyMaterialInstanceParameters(OldMaterialInstanceDynamic);
+		}
+
+		NewMaterialInterface->MarkPackageDirty();
+
+		FAssetRegistryModule::AssetCreated(NewMaterialInterface);
+
+		InMaterialInterface = NewMaterialInterface;
+	}
+
+	OutMaterials.Add(InMaterialInterface);
+}
+
+// Helper function for ConvertMeshesToStaticMesh
+template <typename ComponentType>
+static void ProcessMaterials(ComponentType* InComponent, const FString& InPackageName, TArray<UMaterialInterface*>& OutMaterials)
+{
+	const int32 NumMaterials = InComponent->GetNumMaterials();
+	for (int32 MaterialIndex = 0; MaterialIndex < NumMaterials; MaterialIndex++)
+	{
+		UMaterialInterface* MaterialInterface = InComponent->GetMaterial(MaterialIndex);
+		AddOrDuplicateMaterial(MaterialInterface, InPackageName, OutMaterials);
+	}
+}
+
+// Helper function for ConvertMeshesToStaticMesh
+static bool IsValidSkinnedMeshComponent(USkinnedMeshComponent* InComponent)
+{
+	return InComponent && InComponent->MeshObject && InComponent->IsVisible();
+}
+
+/** Helper struct for tracking validity of optional buffers */
+struct FRawMeshTracker
+{
+	FRawMeshTracker()
+		: bValidColors(false)
+	{
+		FMemory::Memset(bValidTexCoords, 0);
+	}
+
+	bool bValidTexCoords[MAX_MESH_TEXTURE_COORDS];
+	bool bValidColors;
+};
+
+// Helper function for ConvertMeshesToStaticMesh
+static void SkinnedMeshToRawMeshes(USkinnedMeshComponent* InSkinnedMeshComponent, int32 InOverallMaxLODs, const FMatrix& InComponentToWorld, const FString& InPackageName, TArray<FRawMeshTracker>& OutRawMeshTrackers, TArray<FRawMesh>& OutRawMeshes, TArray<UMaterialInterface*>& OutMaterials)
+{
+	const int32 BaseMaterialIndex = OutMaterials.Num();
+
+	// Export all LODs to raw meshes
+	const int32 NumLODs = InSkinnedMeshComponent->MeshObject->GetSkeletalMeshResource().LODModels.Num();
+
+	for (int32 OverallLODIndex = 0; OverallLODIndex < InOverallMaxLODs; OverallLODIndex++)
+	{
+		int32 LODIndexRead = FMath::Min(OverallLODIndex, NumLODs - 1);
+
+		FRawMesh& RawMesh = OutRawMeshes[OverallLODIndex];
+		FRawMeshTracker& RawMeshTracker = OutRawMeshTrackers[OverallLODIndex];
+		const int32 BaseVertexIndex = RawMesh.VertexPositions.Num();
+
+		// Get the CPU skinned verts for this LOD
+		TArray<FFinalSkinVertex> FinalVertices;
+		InSkinnedMeshComponent->GetCPUSkinnedVertices(FinalVertices, LODIndexRead);
+
+		FSkeletalMeshResource& SkeletalMeshResource = InSkinnedMeshComponent->MeshObject->GetSkeletalMeshResource();
+		FStaticLODModel& StaticLODModel = SkeletalMeshResource.LODModels[LODIndexRead];
+
+		// Copy skinned vertex positions
+		for (int32 VertIndex = 0; VertIndex < FinalVertices.Num(); ++VertIndex)
+		{
+			RawMesh.VertexPositions.Add(InComponentToWorld.TransformPosition(FinalVertices[VertIndex].Position));
+		}
+
+		const uint32 NumTexCoords = FMath::Min(StaticLODModel.VertexBufferGPUSkin.GetNumTexCoords(), (uint32)MAX_MESH_TEXTURE_COORDS);
+		const int32 NumSections = StaticLODModel.Sections.Num();
+		FRawStaticIndexBuffer16or32Interface& IndexBuffer = *StaticLODModel.MultiSizeIndexContainer.GetIndexBuffer();
+
+		for (int32 SectionIndex = 0; SectionIndex < NumSections; SectionIndex++)
+		{
+			const FSkelMeshSection& SkelMeshSection = StaticLODModel.Sections[SectionIndex];
+			if (!SkelMeshSection.bDisabled)
+			{
+				// Build 'wedge' info
+				const int32 NumIndices = SkelMeshSection.NumTriangles * 3;
+				for (int32 IndexIndex = 0; IndexIndex < NumIndices; IndexIndex++)
+				{
+					int32 Index = IndexBuffer.Get(SkelMeshSection.BaseIndex + IndexIndex);
+
+					RawMesh.WedgeIndices.Add(BaseVertexIndex + Index);
+
+					const FFinalSkinVertex& SkinnedVertex = FinalVertices[Index];
+					const FVector TangentX = InComponentToWorld.TransformVector(SkinnedVertex.TangentX);
+					const FVector4 TangentZ = InComponentToWorld.TransformVector(SkinnedVertex.TangentZ);
+					const FVector TangentY = (TangentX ^ TangentZ).GetSafeNormal() * TangentZ.W;
+
+					RawMesh.WedgeTangentX.Add(TangentX);
+					RawMesh.WedgeTangentY.Add(TangentY);
+					RawMesh.WedgeTangentZ.Add(TangentZ);
+
+					for (uint32 TexCoordIndex = 0; TexCoordIndex < MAX_MESH_TEXTURE_COORDS; TexCoordIndex++)
+					{
+						if (TexCoordIndex >= NumTexCoords)
+						{
+							RawMesh.WedgeTexCoords[TexCoordIndex].AddDefaulted();
+						}
+						else
+						{
+							RawMesh.WedgeTexCoords[TexCoordIndex].Add(StaticLODModel.VertexBufferGPUSkin.GetVertexUV(Index, TexCoordIndex));
+							RawMeshTracker.bValidTexCoords[TexCoordIndex] = true;
+						}
+					}
+
+					if (StaticLODModel.ColorVertexBuffer.IsInitialized())
+					{
+						RawMesh.WedgeColors.Add(StaticLODModel.ColorVertexBuffer.VertexColor(Index));
+						RawMeshTracker.bValidColors = true;
+					}
+					else
+					{
+						RawMesh.WedgeColors.Add(FColor::White);
+					}
+				}
+
+				// copy face info
+				for (uint32 TriIndex = 0; TriIndex < SkelMeshSection.NumTriangles; TriIndex++)
+				{
+					RawMesh.FaceMaterialIndices.Add(BaseMaterialIndex + SkelMeshSection.MaterialIndex);
+					RawMesh.FaceSmoothingMasks.Add(0); // Assume this is ignored as bRecomputeNormals is false
+				}
+			}
+		}
+	}
+
+	ProcessMaterials<USkinnedMeshComponent>(InSkinnedMeshComponent, InPackageName, OutMaterials);
+}
+
+// Helper function for ConvertMeshesToStaticMesh
+static bool IsValidStaticMeshComponent(UStaticMeshComponent* InComponent)
+{
+	return InComponent && InComponent->StaticMesh && InComponent->StaticMesh->RenderData && InComponent->IsVisible();
+}
+
+// Helper function for ConvertMeshesToStaticMesh
+static void StaticMeshToRawMeshes(UStaticMeshComponent* InStaticMeshComponent, int32 InOverallMaxLODs, const FMatrix& InComponentToWorld, const FString& InPackageName, TArray<FRawMeshTracker>& OutRawMeshTrackers, TArray<FRawMesh>& OutRawMeshes, TArray<UMaterialInterface*>& OutMaterials)
+{
+	const int32 BaseMaterialIndex = OutMaterials.Num();
+
+	const int32 NumLODs = InStaticMeshComponent->StaticMesh->RenderData->LODResources.Num();
+
+	for (int32 OverallLODIndex = 0; OverallLODIndex < InOverallMaxLODs; OverallLODIndex++)
+	{
+		int32 LODIndexRead = FMath::Min(OverallLODIndex, NumLODs - 1);
+
+		FRawMesh& RawMesh = OutRawMeshes[OverallLODIndex];
+		FRawMeshTracker& RawMeshTracker = OutRawMeshTrackers[OverallLODIndex];
+		const FStaticMeshLODResources& LODResource = InStaticMeshComponent->StaticMesh->RenderData->LODResources[LODIndexRead];
+		const int32 BaseVertexIndex = RawMesh.VertexPositions.Num();
+
+		for (int32 VertIndex = 0; VertIndex < LODResource.GetNumVertices(); ++VertIndex)
+		{
+			RawMesh.VertexPositions.Add(InComponentToWorld.TransformPosition(LODResource.PositionVertexBuffer.VertexPosition((uint32)VertIndex)));
+		}
+
+		const FIndexArrayView IndexArrayView = LODResource.IndexBuffer.GetArrayView();
+		const FStaticMeshVertexBuffer& StaticMeshVertexBuffer = LODResource.VertexBuffer;
+		const int32 NumTexCoords = FMath::Min(StaticMeshVertexBuffer.GetNumTexCoords(), (uint32)MAX_MESH_TEXTURE_COORDS);
+		const int32 NumSections = LODResource.Sections.Num();
+
+		for (int32 SectionIndex = 0; SectionIndex < NumSections; SectionIndex++)
+		{
+			const FStaticMeshSection& StaticMeshSection = LODResource.Sections[SectionIndex];
+
+			const int32 NumIndices = StaticMeshSection.NumTriangles * 3;
+			for (int32 IndexIndex = 0; IndexIndex < NumIndices; IndexIndex++)
+			{
+				int32 Index = IndexArrayView[StaticMeshSection.FirstIndex + IndexIndex];
+				RawMesh.WedgeIndices.Add(BaseVertexIndex + Index);
+
+				RawMesh.WedgeTangentX.Add(InComponentToWorld.TransformVector(StaticMeshVertexBuffer.VertexTangentX(Index)));
+				RawMesh.WedgeTangentY.Add(InComponentToWorld.TransformVector(StaticMeshVertexBuffer.VertexTangentY(Index)));
+				RawMesh.WedgeTangentZ.Add(InComponentToWorld.TransformVector(StaticMeshVertexBuffer.VertexTangentZ(Index)));
+
+				for (int32 TexCoordIndex = 0; TexCoordIndex < MAX_MESH_TEXTURE_COORDS; TexCoordIndex++)
+				{
+					if (TexCoordIndex >= NumTexCoords)
+					{
+						RawMesh.WedgeTexCoords[TexCoordIndex].AddDefaulted();
+					}
+					else
+					{
+						RawMesh.WedgeTexCoords[TexCoordIndex].Add(StaticMeshVertexBuffer.GetVertexUV(Index, TexCoordIndex));
+						RawMeshTracker.bValidTexCoords[TexCoordIndex] = true;
+					}
+				}
+
+				if (LODResource.ColorVertexBuffer.IsInitialized())
+				{
+					RawMesh.WedgeColors.Add(LODResource.ColorVertexBuffer.VertexColor(Index));
+					RawMeshTracker.bValidColors = true;
+				}
+				else
+				{
+					RawMesh.WedgeColors.Add(FColor::White);
+				}
+			}
+
+			// copy face info
+			for (uint32 TriIndex = 0; TriIndex < StaticMeshSection.NumTriangles; TriIndex++)
+			{
+				RawMesh.FaceMaterialIndices.Add(BaseMaterialIndex + StaticMeshSection.MaterialIndex);
+				RawMesh.FaceSmoothingMasks.Add(0); // Assume this is ignored as bRecomputeNormals is false
+			}
+		}
+	}
+
+	ProcessMaterials<UStaticMeshComponent>(InStaticMeshComponent, InPackageName, OutMaterials);
+}
+
+UStaticMesh* FMeshUtilities::ConvertMeshesToStaticMesh(const TArray<UMeshComponent*>& InMeshComponents, const FTransform& InRootTransform, const FString& InPackageName)
+{
+	// Build a package name to use
+	FString MeshName;
+	FString PackageName;
+	if (InPackageName.IsEmpty())
+	{
+		FString NewNameSuggestion = FString(TEXT("StaticMesh"));
+		FString PackageNameSuggestion = FString(TEXT("/Game/Meshes/")) + NewNameSuggestion;
+		FString Name;
+		FAssetToolsModule& AssetToolsModule = FModuleManager::LoadModuleChecked<FAssetToolsModule>("AssetTools");
+		AssetToolsModule.Get().CreateUniqueAssetName(PackageNameSuggestion, TEXT(""), PackageNameSuggestion, Name);
+
+		TSharedPtr<SDlgPickAssetPath> PickAssetPathWidget =
+			SNew(SDlgPickAssetPath)
+			.Title(LOCTEXT("ConvertToStaticMeshPickName", "Choose New StaticMesh Location"))
+			.DefaultAssetPath(FText::FromString(PackageNameSuggestion));
+
+		if (PickAssetPathWidget->ShowModal() == EAppReturnType::Ok)
+		{
+			// Get the full name of where we want to create the mesh asset.
+			PackageName = PickAssetPathWidget->GetFullAssetPath().ToString();
+			MeshName = FPackageName::GetLongPackageAssetName(PackageName);
+
+			// Check if the user inputed a valid asset name, if they did not, give it the generated default name
+			if (MeshName.IsEmpty())
+			{
+				// Use the defaults that were already generated.
+				PackageName = PackageNameSuggestion;
+				MeshName = *Name;
+			}
+		}
+	}
+	else
+	{
+		PackageName = InPackageName;
+		MeshName = *FPackageName::GetLongPackageAssetName(PackageName);
+	}
+
+	if(!PackageName.IsEmpty() && !MeshName.IsEmpty())
+	{
+		TArray<FRawMesh> RawMeshes;
+		TArray<UMaterialInterface*> Materials;
+
+		TArray<FRawMeshTracker> RawMeshTrackers;
+
+		FMatrix WorldToRoot = InRootTransform.ToMatrixWithScale().Inverse();
+
+		FScopedSlowTask ScopedSlowTask((float)InMeshComponents.Num() + 1.0f, LOCTEXT("MeshConversionProgress", "Converting Meshes"));
+		ScopedSlowTask.MakeDialog(false, true);
+
+		// first do a pass to determine the max LOD level we will be combining meshes into
+		int32 OverallMaxLODs = 0;
+		for (UMeshComponent* MeshComponent : InMeshComponents)
+		{
+			USkinnedMeshComponent* SkinnedMeshComponent = Cast<USkinnedMeshComponent>(MeshComponent);
+			UStaticMeshComponent* StaticMeshComponent = Cast<UStaticMeshComponent>(MeshComponent);
+
+			if (IsValidSkinnedMeshComponent(SkinnedMeshComponent))
+			{
+				OverallMaxLODs = FMath::Max(SkinnedMeshComponent->MeshObject->GetSkeletalMeshResource().LODModels.Num(), OverallMaxLODs);
+			}
+			else if(IsValidStaticMeshComponent(StaticMeshComponent))
+			{
+				OverallMaxLODs = FMath::Max(StaticMeshComponent->StaticMesh->RenderData->LODResources.Num(), OverallMaxLODs);
+			}
+		}
+		
+		// Resize raw meshes to accommodate the number of LODs we will need
+		RawMeshes.SetNum(OverallMaxLODs);
+		RawMeshTrackers.SetNum(OverallMaxLODs);
+
+		// Export all visible components
+		for (UMeshComponent* MeshComponent : InMeshComponents)
+		{
+			ScopedSlowTask.EnterProgressFrame(1.0f);
+
+			FMatrix ComponentToWorld = MeshComponent->GetComponentTransform().ToMatrixWithScale() * WorldToRoot;
+
+			USkinnedMeshComponent* SkinnedMeshComponent = Cast<USkinnedMeshComponent>(MeshComponent);
+			UStaticMeshComponent* StaticMeshComponent = Cast<UStaticMeshComponent>(MeshComponent);
+
+			if (IsValidSkinnedMeshComponent(SkinnedMeshComponent))
+			{
+				SkinnedMeshToRawMeshes(SkinnedMeshComponent, OverallMaxLODs, ComponentToWorld, PackageName, RawMeshTrackers, RawMeshes, Materials);
+			}
+			else if (IsValidStaticMeshComponent(StaticMeshComponent))
+			{
+				StaticMeshToRawMeshes(StaticMeshComponent, OverallMaxLODs, ComponentToWorld, PackageName, RawMeshTrackers, RawMeshes, Materials);
+			}
+		}
+
+		// scrub invalid vert color & tex coord data
+		check(RawMeshes.Num() == RawMeshTrackers.Num());
+		for (int32 RawMeshIndex = 0; RawMeshIndex < RawMeshes.Num(); RawMeshIndex++)
+		{
+			if (!RawMeshTrackers[RawMeshIndex].bValidColors)
+			{
+				RawMeshes[RawMeshIndex].WedgeColors.Empty();
+			}
+
+			for (uint32 TexCoordIndex = 0; TexCoordIndex < MAX_MESH_TEXTURE_COORDS; TexCoordIndex++)
+			{
+				if (!RawMeshTrackers[RawMeshIndex].bValidTexCoords[TexCoordIndex])
+				{
+					RawMeshes[RawMeshIndex].WedgeTexCoords[TexCoordIndex].Empty();
+				}
+			}
+		}
+
+		// Check if we got some valid data.
+		bool bValidData = false;
+		for (FRawMesh& RawMesh : RawMeshes)
+		{
+			if (RawMesh.IsValidOrFixable())
+			{
+				bValidData = true;
+				break;
+			}
+		}
+
+		if (bValidData)
+		{
+			ScopedSlowTask.EnterProgressFrame(1.0f, LOCTEXT("FinalizingMesh", "Finalizing Mesh"));
+
+			// Then find/create it.
+			UPackage* Package = CreatePackage(NULL, *PackageName);
+			check(Package);
+
+			// Create StaticMesh object
+			UStaticMesh* StaticMesh = NewObject<UStaticMesh>(Package, *MeshName, RF_Public | RF_Standalone);
+			StaticMesh->InitResources();
+
+			StaticMesh->LightingGuid = FGuid::NewGuid();
+
+			// Add source to new StaticMesh
+			for (FRawMesh& RawMesh : RawMeshes)
+			{
+				if (RawMesh.IsValidOrFixable())
+				{
+					FStaticMeshSourceModel* SrcModel = new (StaticMesh->SourceModels) FStaticMeshSourceModel();
+					SrcModel->BuildSettings.bRecomputeNormals = false;
+					SrcModel->BuildSettings.bRecomputeTangents = false;
+					SrcModel->BuildSettings.bRemoveDegenerates = false;
+					SrcModel->BuildSettings.bUseHighPrecisionTangentBasis = false;
+					SrcModel->BuildSettings.bUseFullPrecisionUVs = false;
+					SrcModel->BuildSettings.bGenerateLightmapUVs = true;
+					SrcModel->BuildSettings.SrcLightmapIndex = 0;
+					SrcModel->BuildSettings.DstLightmapIndex = 1;
+					SrcModel->RawMeshBulkData->SaveRawMesh(RawMesh);
+				}
+			}
+
+			// Copy materials to new mesh 
+			for(UMaterialInterface* Material : Materials)
+			{
+				StaticMesh->StaticMaterials.Add(FStaticMaterial(Material));
+			}
+			
+			// Build mesh from source
+			StaticMesh->Build(false);
+			StaticMesh->PostEditChange();
+
+			StaticMesh->MarkPackageDirty();
+
+			// Notify asset registry of new asset
+			FAssetRegistryModule::AssetCreated(StaticMesh);
+
+			// Display notification so users can quickly access the mesh
+			if (GIsEditor)
+			{
+				FNotificationInfo Info(FText::Format(LOCTEXT("SkeletalMeshConverted", "Successfully Converted Mesh"), FText::FromString(StaticMesh->GetName())));
+				Info.ExpireDuration = 8.0f;
+				Info.bUseLargeFont = false;
+				Info.Hyperlink = FSimpleDelegate::CreateLambda([=]() { FAssetEditorManager::Get().OpenEditorForAssets(TArray<UObject*>({ StaticMesh })); });
+				Info.HyperlinkText = FText::Format(LOCTEXT("OpenNewAnimationHyperlink", "Open {0}"), FText::FromString(StaticMesh->GetName()));
+				TSharedPtr<SNotificationItem> Notification = FSlateNotificationManager::Get().AddNotification(Info);
+				if ( Notification.IsValid() )
+				{
+					Notification->SetCompletionState( SNotificationItem::CS_Success );
+				}
+			}
+		}
+	}
+
+	return nullptr;
+}
+
 /**
 * Builds a renderable skeletal mesh LOD model. Note that the array of chunks
 * will be destroyed during this process!
@@ -5415,8 +5906,8 @@ void FMeshUtilities::CreateProxyMesh(const TArray<AActor*>& InActors, const stru
 			// Add constructed raw mesh to source mesh array
 			const int32 SourceMeshIndex = SourceMeshes.AddZeroed();
 			SourceMeshes[SourceMeshIndex].MeshLODData[ProxyMeshTargetLODLevel].RawMesh = RawMesh;
-			SourceMeshes[SourceMeshIndex].bShouldExportLOD[ProxyMeshSourceLODLevel] = true;
-			SourceMeshes[SourceMeshIndex].ExportLODIndex = ProxyMeshSourceLODLevel;
+			SourceMeshes[SourceMeshIndex].bShouldExportLOD[ProxyMeshTargetLODLevel] = true;
+			SourceMeshes[SourceMeshIndex].ExportLODIndex = ProxyMeshTargetLODLevel;
 
 			// Append retrieved materials for this static mesh component to the global material map
 			GlobalMaterialMap.Add(FMeshIdAndLOD(SourceMeshIndex, ProxyMeshTargetLODLevel), StaticMeshGlobalMaterialMap);
@@ -5704,7 +6195,7 @@ void FMeshUtilities::FlattenMaterialsWithMeshData(TArray<UMaterialInterface*>& I
 		float MaxEmissiveScale = 0.0f;
 		for (FFlattenMaterial& FlatMaterial : OutFlattenedMaterials)
 		{
-			if (FlatMaterial.EmissiveSamples.Num())
+			if (FlatMaterial.DoesPropertyContainData(EFlattenMaterialProperties::Emissive))
 			{
 				if (FlatMaterial.EmissiveScale > MaxEmissiveScale)
 				{
@@ -5725,9 +6216,10 @@ void FMeshUtilities::FlattenMaterialsWithMeshData(TArray<UMaterialInterface*>& I
 					continue;
 				}
 				// Rescale emissive data.
-				for (int32 PixelIndex = 0; PixelIndex < FlatMaterial.EmissiveSamples.Num(); PixelIndex++)
+				TArray<FColor>& EmissiveSamples = FlatMaterial.GetPropertySamples(EFlattenMaterialProperties::Emissive);
+				for (int32 PixelIndex = 0; PixelIndex < EmissiveSamples.Num(); PixelIndex++)
 				{
-					FColor& C = FlatMaterial.EmissiveSamples[PixelIndex];
+					FColor& C = EmissiveSamples[PixelIndex];
 					C.R = FMath::RoundToInt(C.R * Scale);
 					C.G = FMath::RoundToInt(C.G * Scale);
 					C.B = FMath::RoundToInt(C.B * Scale);
@@ -6662,37 +7154,13 @@ static FVector2D GetValidUV(const FVector2D& UV)
 	return NewUV;
 }
 
-static void MergeMaterials(UWorld* InWorld, const TArray<UMaterialInterface*>& InMaterialList, FFlattenMaterial& OutMergedMaterial, TArray<FRawMeshUVTransform>& OutUVTransforms)
+static void MergeFlattenedMaterials(TArray<struct FFlattenMaterial>& InMaterialList, FFlattenMaterial& OutMergedMaterial, TArray<FRawMeshUVTransform>& OutUVTransforms)
 {
 	OutUVTransforms.Reserve(InMaterialList.Num());
 
-#if !(UE_BUILD_SHIPPING || UE_BUILD_TEST)
-	// add log to warn which material it
-	UE_LOG(LogMeshUtilities, Log, TEXT("Merging Material: Merge Material count %d"), InMaterialList.Num());
-
-	int32 Index = 0;
-	for (auto& MaterialIter : InMaterialList)
-	{
-		if (MaterialIter)
-		{
-			UE_LOG(LogMeshUtilities, Log, TEXT("Material List %d - %s"), ++Index, *MaterialIter->GetName());
-		}
-		else
-		{
-			UE_LOG(LogMeshUtilities, Log, TEXT("Material List %d - null material"), ++Index);
-		}
-	}
-
-#endif
-	// We support merging only for opaque materials
-	int32 NumOpaqueMaterials = 0;
 	// Fill output UV transforms with invalid values
 	for (auto Material : InMaterialList)
 	{
-		if (Material->GetBlendMode() == BLEND_Opaque)
-		{
-			NumOpaqueMaterials++;
-		}
 
 		// Invalid UV transform
 		FRawMeshUVTransform UVTransform;
@@ -6701,48 +7169,21 @@ static void MergeMaterials(UWorld* InWorld, const TArray<UMaterialInterface*>& I
 		OutUVTransforms.Add(UVTransform);
 	}
 
-	if (NumOpaqueMaterials == 0)
-	{
-		// Nothing to merge
-		return;
-	}
-
-	int32 AtlasGridSize = FMath::CeilToInt(FMath::Sqrt(NumOpaqueMaterials));
-	FIntPoint AtlasTextureSize = OutMergedMaterial.DiffuseSize;
+	int32 AtlasGridSize = FMath::CeilToInt(FMath::Sqrt(InMaterialList.Num()));
+	FIntPoint AtlasTextureSize = OutMergedMaterial.GetPropertySize(EFlattenMaterialProperties::Diffuse);
 	FIntPoint ExportTextureSize = AtlasTextureSize / AtlasGridSize;
 	int32 AtlasNumSamples = AtlasTextureSize.X*AtlasTextureSize.Y;
+	
 
-	bool bExportNormal = (OutMergedMaterial.NormalSize != FIntPoint::ZeroValue);
-	bool bExportMetallic = (OutMergedMaterial.MetallicSize != FIntPoint::ZeroValue);
-	bool bExportRoughness = (OutMergedMaterial.RoughnessSize != FIntPoint::ZeroValue);
-	bool bExportSpecular = (OutMergedMaterial.SpecularSize != FIntPoint::ZeroValue);
-
-	// Pre-allocate buffers for texture atlas
-	OutMergedMaterial.DiffuseSamples.Reserve(AtlasNumSamples);
-	OutMergedMaterial.DiffuseSamples.SetNumZeroed(AtlasNumSamples);
-	if (bExportNormal)
-	{
-		check(OutMergedMaterial.NormalSize == OutMergedMaterial.DiffuseSize);
-		OutMergedMaterial.NormalSamples.Reserve(AtlasNumSamples);
-		OutMergedMaterial.NormalSamples.SetNumZeroed(AtlasNumSamples);
-	}
-	if (bExportMetallic)
-	{
-		check(OutMergedMaterial.MetallicSize == OutMergedMaterial.DiffuseSize);
-		OutMergedMaterial.MetallicSamples.Reserve(AtlasNumSamples);
-		OutMergedMaterial.MetallicSamples.SetNumZeroed(AtlasNumSamples);
-	}
-	if (bExportRoughness)
-	{
-		check(OutMergedMaterial.RoughnessSize == OutMergedMaterial.DiffuseSize);
-		OutMergedMaterial.RoughnessSamples.Reserve(AtlasNumSamples);
-		OutMergedMaterial.RoughnessSamples.SetNumZeroed(AtlasNumSamples);
-	}
-	if (bExportSpecular)
+	for (int32 PropertyIndex = 0; PropertyIndex < (int32)EFlattenMaterialProperties::NumFlattenMaterialProperties; ++PropertyIndex)
 	{
-		check(OutMergedMaterial.SpecularSize == OutMergedMaterial.DiffuseSize);
-		OutMergedMaterial.SpecularSamples.Reserve(AtlasNumSamples);
-		OutMergedMaterial.SpecularSamples.SetNumZeroed(AtlasNumSamples);
+		EFlattenMaterialProperties Property = (EFlattenMaterialProperties)PropertyIndex;
+		if (OutMergedMaterial.ShouldGenerateDataForProperty(Property))
+		{
+			check(OutMergedMaterial.GetPropertySize(Property) == AtlasTextureSize);
+			TArray<FColor>& Samples = OutMergedMaterial.GetPropertySamples(Property);
+			Samples.SetNumZeroed(AtlasNumSamples);
+		}
 	}
 
 	int32 AtlasRowIdx = 0;
@@ -6752,65 +7193,27 @@ static void MergeMaterials(UWorld* InWorld, const TArray<UMaterialInterface*>& I
 	// Flatten all materials and merge them into one material using texture atlases
 	for (int32 MatIdx = 0; MatIdx < InMaterialList.Num(); ++MatIdx)
 	{
-		UMaterialInterface* Material = InMaterialList[MatIdx];
-		if (Material->GetBlendMode() != BLEND_Opaque)
-		{
-			continue;
-		}
-
-		FFlattenMaterial FlatMaterial;
-		FlatMaterial.DiffuseSize = ExportTextureSize;
-		FlatMaterial.NormalSize = bExportNormal ? ExportTextureSize : FIntPoint::ZeroValue;
-		FlatMaterial.MetallicSize = bExportMetallic ? ExportTextureSize : FIntPoint::ZeroValue;
-		FlatMaterial.RoughnessSize = bExportRoughness ? ExportTextureSize : FIntPoint::ZeroValue;
-		FlatMaterial.SpecularSize = bExportSpecular ? ExportTextureSize : FIntPoint::ZeroValue;
-		int32 ExportNumSamples = ExportTextureSize.X*ExportTextureSize.Y;
-
-		FMaterialUtilities::ExportMaterial(Material, FlatMaterial);
+		FFlattenMaterial& FlatMaterial = InMaterialList[MatIdx];
 
-		if (FlatMaterial.DiffuseSamples.Num() == ExportNumSamples)
-		{
-			CopyTextureRect(FlatMaterial.DiffuseSamples.GetData(), ExportTextureSize, OutMergedMaterial.DiffuseSamples.GetData(), AtlasTextureSize, AtlasTargetPos);
-		}
-		else if (FlatMaterial.DiffuseSamples.Num() == 1)
+		for (int32 PropertyIndex = 0; PropertyIndex < (int32)EFlattenMaterialProperties::NumFlattenMaterialProperties; ++PropertyIndex)
 		{
-			SetTextureRect(FlatMaterial.DiffuseSamples[0], ExportTextureSize, OutMergedMaterial.DiffuseSamples.GetData(), AtlasTextureSize, AtlasTargetPos);
-		}
-
-		if (FlatMaterial.NormalSamples.Num() == ExportNumSamples)
-		{
-			CopyTextureRect(FlatMaterial.NormalSamples.GetData(), ExportTextureSize, OutMergedMaterial.NormalSamples.GetData(), AtlasTextureSize, AtlasTargetPos);
-		}
-		else if (FlatMaterial.NormalSamples.Num() == 1)
-		{
-			SetTextureRect(FlatMaterial.NormalSamples[0], ExportTextureSize, OutMergedMaterial.NormalSamples.GetData(), AtlasTextureSize, AtlasTargetPos);
-		}
-
-		if (FlatMaterial.MetallicSamples.Num() == ExportNumSamples)
-		{
-			CopyTextureRect(FlatMaterial.MetallicSamples.GetData(), ExportTextureSize, OutMergedMaterial.MetallicSamples.GetData(), AtlasTextureSize, AtlasTargetPos);
-		}
-		else if (FlatMaterial.MetallicSamples.Num() == 1)
-		{
-			SetTextureRect(FlatMaterial.MetallicSamples[0], ExportTextureSize, OutMergedMaterial.MetallicSamples.GetData(), AtlasTextureSize, AtlasTargetPos);
-		}
-
-		if (FlatMaterial.RoughnessSamples.Num() == ExportNumSamples)
-		{
-			CopyTextureRect(FlatMaterial.RoughnessSamples.GetData(), ExportTextureSize, OutMergedMaterial.RoughnessSamples.GetData(), AtlasTextureSize, AtlasTargetPos);
-		}
-		else if (FlatMaterial.RoughnessSamples.Num() == 1)
-		{
-			SetTextureRect(FlatMaterial.RoughnessSamples[0], ExportTextureSize, OutMergedMaterial.RoughnessSamples.GetData(), AtlasTextureSize, AtlasTargetPos);
-		}
-
-		if (FlatMaterial.SpecularSamples.Num() == ExportNumSamples)
-		{
-			CopyTextureRect(FlatMaterial.SpecularSamples.GetData(), ExportTextureSize, OutMergedMaterial.SpecularSamples.GetData(), AtlasTextureSize, AtlasTargetPos);
-		}
-		else if (FlatMaterial.SpecularSamples.Num() == 1)
-		{
-			SetTextureRect(FlatMaterial.SpecularSamples[0], ExportTextureSize, OutMergedMaterial.SpecularSamples.GetData(), AtlasTextureSize, AtlasTargetPos);
+			EFlattenMaterialProperties Property = (EFlattenMaterialProperties)PropertyIndex;
+			if (OutMergedMaterial.ShouldGenerateDataForProperty(Property) && FlatMaterial.DoesPropertyContainData(Property))
+			{
+				TArray<FColor>& SourceSamples = FlatMaterial.GetPropertySamples(Property);
+				TArray<FColor>& TargetSamples = OutMergedMaterial.GetPropertySamples(Property);
+				if (FlatMaterial.IsPropertyConstant(Property))
+				{
+					SetTextureRect(SourceSamples[0], ExportTextureSize, TargetSamples.GetData(), AtlasTextureSize, AtlasTargetPos);
+				}
+				else
+				{
+					FIntPoint PropertySize = FlatMaterial.GetPropertySize(Property);
+					PropertySize = ConditionalImageResize(PropertySize, ExportTextureSize, SourceSamples, false);
+					CopyTextureRect(SourceSamples.GetData(), ExportTextureSize, TargetSamples.GetData(), AtlasTextureSize, AtlasTargetPos);
+					FlatMaterial.SetPropertySize(Property, PropertySize);
+				}
+			}
 		}
 
 		check(OutUVTransforms.IsValidIndex(MatIdx));
@@ -6834,29 +7237,10 @@ static void MergeMaterials(UWorld* InWorld, const TArray<UMaterialInterface*>& I
 	}
 }
 
-
-static void MergeFlattenedMaterials(TArray<struct FFlattenMaterial>& InMaterialList, FFlattenMaterial& OutMergedMaterial, TArray<FRawMeshUVTransform>& OutUVTransforms)
+static void FlattenBinnedMaterials(TArray<struct FFlattenMaterial>& InMaterialList, const TArray<FBox2D>& InMaterialBoxes, FFlattenMaterial& OutMergedMaterial, TArray<FRawMeshUVTransform>& OutUVTransforms)
 {
 	OutUVTransforms.Reserve(InMaterialList.Num());
 
-#if !(UE_BUILD_SHIPPING || UE_BUILD_TEST)
-	//// add log to warn which material it
-	//UE_LOG(LogMeshUtilities, Log, TEXT("Merging Material: Merge Material count %d"), InMaterialList.Num());
-
-	//int32 Index = 0;
-	//for (auto& MaterialIter : InMaterialList)
-	//{
-	//	if (MaterialIter)
-	//	{
-	//		UE_LOG(LogMeshUtilities, Log, TEXT("Material List %d - %s"), ++Index, *MaterialIter->GetName());
-	//	}
-	//	else
-	//	{
-	//		UE_LOG(LogMeshUtilities, Log, TEXT("Material List %d - null material"), ++Index);
-	//	}
-	//}
-
-#endif
 	// We support merging only for opaque materials
 	// Fill output UV transforms with invalid values
 	for (auto Material : InMaterialList)
@@ -6867,159 +7251,47 @@ static void MergeFlattenedMaterials(TArray<struct FFlattenMaterial>& InMaterialL
 		UVTransform.Scale = FVector2D::ZeroVector;
 		OutUVTransforms.Add(UVTransform);
 	}
-
-	int32 AtlasGridSize = FMath::CeilToInt(FMath::Sqrt(InMaterialList.Num()));
-	FIntPoint AtlasTextureSize = OutMergedMaterial.DiffuseSize;
-	FIntPoint ExportTextureSize = AtlasTextureSize / AtlasGridSize;
-	int32 AtlasNumSamples = AtlasTextureSize.X*AtlasTextureSize.Y;
-
-	bool bExportNormal = (OutMergedMaterial.NormalSize != FIntPoint::ZeroValue);
-	bool bExportMetallic = (OutMergedMaterial.MetallicSize != FIntPoint::ZeroValue);
-	bool bExportRoughness = (OutMergedMaterial.RoughnessSize != FIntPoint::ZeroValue);
-	bool bExportSpecular = (OutMergedMaterial.SpecularSize != FIntPoint::ZeroValue);
-	bool bExportEmissive = (OutMergedMaterial.EmissiveSize != FIntPoint::ZeroValue);
-	bool bExportOpacity = (OutMergedMaterial.OpacitySize != FIntPoint::ZeroValue);
-
-	// Pre-allocate buffers for texture atlas
-	OutMergedMaterial.DiffuseSamples.Reserve(AtlasNumSamples);
-	OutMergedMaterial.DiffuseSamples.SetNumZeroed(AtlasNumSamples);
-	if (bExportNormal)
-	{
-		check(OutMergedMaterial.NormalSize == OutMergedMaterial.DiffuseSize);
-		OutMergedMaterial.NormalSamples.Reserve(AtlasNumSamples);
-		OutMergedMaterial.NormalSamples.SetNumZeroed(AtlasNumSamples);
-	}
-	if (bExportMetallic)
-	{
-		check(OutMergedMaterial.MetallicSize == OutMergedMaterial.DiffuseSize);
-		OutMergedMaterial.MetallicSamples.Reserve(AtlasNumSamples);
-		OutMergedMaterial.MetallicSamples.SetNumZeroed(AtlasNumSamples);
-	}
-	if (bExportRoughness)
-	{
-		check(OutMergedMaterial.RoughnessSize == OutMergedMaterial.DiffuseSize);
-		OutMergedMaterial.RoughnessSamples.Reserve(AtlasNumSamples);
-		OutMergedMaterial.RoughnessSamples.SetNumZeroed(AtlasNumSamples);
-	}
-	if (bExportSpecular)
-	{
-		check(OutMergedMaterial.SpecularSize == OutMergedMaterial.DiffuseSize);
-		OutMergedMaterial.SpecularSamples.Reserve(AtlasNumSamples);
-		OutMergedMaterial.SpecularSamples.SetNumZeroed(AtlasNumSamples);
-	}
-	if (bExportEmissive)
-	{
-		check(OutMergedMaterial.EmissiveSize == OutMergedMaterial.EmissiveSize);
-		OutMergedMaterial.EmissiveSamples.Reserve(AtlasNumSamples);
-		OutMergedMaterial.EmissiveSamples.SetNumZeroed(AtlasNumSamples);
-	}
-
-	if (bExportOpacity)
-	{
-		check(OutMergedMaterial.OpacitySize == OutMergedMaterial.OpacitySize);
-		OutMergedMaterial.OpacitySamples.Reserve(AtlasNumSamples);
-		OutMergedMaterial.OpacitySamples.SetNumZeroed(AtlasNumSamples);
-	}
-
-
-	int32 AtlasRowIdx = 0;
-	int32 AtlasColIdx = 0;
-	FIntPoint AtlasTargetPos = FIntPoint(0, 0);
-
-	// Flatten all materials and merge them into one material using texture atlases
-	for (int32 MatIdx = 0; MatIdx < InMaterialList.Num(); ++MatIdx)
-	{
-		FFlattenMaterial& FlatMaterial = InMaterialList[MatIdx];
-
-		if (FlatMaterial.DiffuseSamples.Num() >= 1)
-		{
-			FlatMaterial.DiffuseSize = ConditionalImageResize(FlatMaterial.DiffuseSize, ExportTextureSize, FlatMaterial.DiffuseSamples, false);
-			CopyTextureRect(FlatMaterial.DiffuseSamples.GetData(), ExportTextureSize, OutMergedMaterial.DiffuseSamples.GetData(), AtlasTextureSize, AtlasTargetPos);
-		}
-		else if (FlatMaterial.DiffuseSamples.Num() == 1)
-		{
-			SetTextureRect(FlatMaterial.DiffuseSamples[0], ExportTextureSize, OutMergedMaterial.DiffuseSamples.GetData(), AtlasTextureSize, AtlasTargetPos);
-		}
-
-		if (FlatMaterial.NormalSamples.Num() >= 1 && bExportNormal)
-		{
-			FlatMaterial.NormalSize = ConditionalImageResize(FlatMaterial.NormalSize, ExportTextureSize, FlatMaterial.NormalSamples, false);
-			CopyTextureRect(FlatMaterial.NormalSamples.GetData(), ExportTextureSize, OutMergedMaterial.NormalSamples.GetData(), AtlasTextureSize, AtlasTargetPos);
-		}
-		else if (FlatMaterial.NormalSamples.Num() == 1 && bExportNormal)
-		{
-			SetTextureRect(FlatMaterial.NormalSamples[0], ExportTextureSize, OutMergedMaterial.NormalSamples.GetData(), AtlasTextureSize, AtlasTargetPos);
-		}
-
-		if (FlatMaterial.MetallicSamples.Num() >= 1 && bExportMetallic)
-		{
-			FlatMaterial.MetallicSize = ConditionalImageResize(FlatMaterial.MetallicSize, ExportTextureSize, FlatMaterial.MetallicSamples, false);
-			CopyTextureRect(FlatMaterial.MetallicSamples.GetData(), ExportTextureSize, OutMergedMaterial.MetallicSamples.GetData(), AtlasTextureSize, AtlasTargetPos);
-		}
-		else if (FlatMaterial.MetallicSamples.Num() == 1 && bExportMetallic)
-		{
-			SetTextureRect(FlatMaterial.MetallicSamples[0], ExportTextureSize, OutMergedMaterial.MetallicSamples.GetData(), AtlasTextureSize, AtlasTargetPos);
-		}
-
-		if (FlatMaterial.RoughnessSamples.Num() >= 1 && bExportRoughness)
-		{
-			FlatMaterial.RoughnessSize = ConditionalImageResize(FlatMaterial.RoughnessSize, ExportTextureSize, FlatMaterial.RoughnessSamples, false);
-			CopyTextureRect(FlatMaterial.RoughnessSamples.GetData(), ExportTextureSize, OutMergedMaterial.RoughnessSamples.GetData(), AtlasTextureSize, AtlasTargetPos);
-		}
-		else if (FlatMaterial.RoughnessSamples.Num() == 1 && bExportRoughness)
-		{
-			SetTextureRect(FlatMaterial.RoughnessSamples[0], ExportTextureSize, OutMergedMaterial.RoughnessSamples.GetData(), AtlasTextureSize, AtlasTargetPos);
-		}
-
-		if (FlatMaterial.SpecularSamples.Num() >= 1 && bExportSpecular)
-		{
-			FlatMaterial.SpecularSize = ConditionalImageResize(FlatMaterial.SpecularSize, ExportTextureSize, FlatMaterial.SpecularSamples, false);
-			CopyTextureRect(FlatMaterial.SpecularSamples.GetData(), ExportTextureSize, OutMergedMaterial.SpecularSamples.GetData(), AtlasTextureSize, AtlasTargetPos);
-		}
-		else if (FlatMaterial.SpecularSamples.Num() == 1 && bExportSpecular)
-		{
-			SetTextureRect(FlatMaterial.SpecularSamples[0], ExportTextureSize, OutMergedMaterial.SpecularSamples.GetData(), AtlasTextureSize, AtlasTargetPos);
-		}
-
-		if (FlatMaterial.EmissiveSamples.Num() >= 1)
-		{
-			FlatMaterial.SpecularSize = bExportEmissive ? ConditionalImageResize(FlatMaterial.EmissiveSize, ExportTextureSize, FlatMaterial.EmissiveSamples, false) : FIntPoint::ZeroValue;
-			CopyTextureRect(FlatMaterial.EmissiveSamples.GetData(), ExportTextureSize, OutMergedMaterial.EmissiveSamples.GetData(), AtlasTextureSize, AtlasTargetPos);
-		}
-		else if (FlatMaterial.EmissiveSamples.Num() == 1)
-		{
-			SetTextureRect(FlatMaterial.EmissiveSamples[0], ExportTextureSize, OutMergedMaterial.EmissiveSamples.GetData(), AtlasTextureSize, AtlasTargetPos);
-		}
-
-		if (FlatMaterial.OpacitySamples.Num() >= 1)
-		{
-			FlatMaterial.SpecularSize = bExportOpacity ? ConditionalImageResize(FlatMaterial.OpacitySize, ExportTextureSize, FlatMaterial.OpacitySamples, false) : FIntPoint::ZeroValue;
-			CopyTextureRect(FlatMaterial.OpacitySamples.GetData(), ExportTextureSize, OutMergedMaterial.OpacitySamples.GetData(), AtlasTextureSize, AtlasTargetPos);
-		}
-		else if (FlatMaterial.OpacitySamples.Num() == 1)
-		{
-			SetTextureRect(FlatMaterial.OpacitySamples[0], ExportTextureSize, OutMergedMaterial.OpacitySamples.GetData(), AtlasTextureSize, AtlasTargetPos);
-		}
-
-
-		check(OutUVTransforms.IsValidIndex(MatIdx));
-
-		OutUVTransforms[MatIdx].Offset = FVector2D(
-			(float)AtlasTargetPos.X / AtlasTextureSize.X,
-			(float)AtlasTargetPos.Y / AtlasTextureSize.Y);
-
-		OutUVTransforms[MatIdx].Scale = FVector2D(
-			(float)ExportTextureSize.X / AtlasTextureSize.X,
-			(float)ExportTextureSize.Y / AtlasTextureSize.Y);
-
-		AtlasColIdx++;
-		if (AtlasColIdx >= AtlasGridSize)
-		{
-			AtlasColIdx = 0;
-			AtlasRowIdx++;
-		}
-
-		AtlasTargetPos = FIntPoint(AtlasColIdx*ExportTextureSize.X, AtlasRowIdx*ExportTextureSize.Y);
+	
+	// Merge all material properties
+	for (int32 Index = 0; Index < (int32)EFlattenMaterialProperties::NumFlattenMaterialProperties; ++Index)
+	{
+		const EFlattenMaterialProperties Property = (EFlattenMaterialProperties)Index;
+		const FIntPoint& OutTextureSize = OutMergedMaterial.GetPropertySize(Property);
+		if (OutTextureSize != FIntPoint::ZeroValue)
+		{
+			TArray<FColor>& OutSamples = OutMergedMaterial.GetPropertySamples(Property);
+			OutSamples.Reserve(OutTextureSize.X * OutTextureSize.Y);
+			OutSamples.SetNumZeroed(OutTextureSize.X * OutTextureSize.Y);
+
+			for (int32 MaterialIndex = 0; MaterialIndex < InMaterialList.Num(); ++MaterialIndex)
+			{
+				// Determine output size and offset
+				FFlattenMaterial& FlatMaterial = InMaterialList[MaterialIndex];
+								
+				if (FlatMaterial.DoesPropertyContainData(Property))
+				{					
+					FBox2D MaterialBox = InMaterialBoxes[MaterialIndex];
+					const FIntPoint& InputSize = FlatMaterial.GetPropertySize(Property);
+					TArray<FColor>& InputSamples = FlatMaterial.GetPropertySamples(Property);
+
+					// Resize material to match output (area) size
+					FIntPoint OutputSize = FIntPoint(OutTextureSize.X * MaterialBox.GetSize().X, OutTextureSize.Y * MaterialBox.GetSize().Y);
+					ConditionalImageResize(InputSize, OutputSize, InputSamples, false);
+
+					// Copy material data to the merged 'atlas' texture
+					FIntPoint OutputPosition = FIntPoint(OutTextureSize.X * MaterialBox.Min.X, OutTextureSize.Y * MaterialBox.Min.Y);
+					CopyTextureRect(InputSamples.GetData(), OutputSize, OutSamples.GetData(), OutTextureSize, OutputPosition);
+
+					// Set the UV tranforms only once
+					if (Index == 0)
+					{
+						FRawMeshUVTransform& UVTransform = OutUVTransforms[MaterialIndex];
+						UVTransform.Offset = MaterialBox.Min;
+						UVTransform.Scale = MaterialBox.GetSize();
+					}
+				}
+			}
+		}		
 	}
 }
 
@@ -7273,16 +7545,28 @@ void FMeshUtilities::MergeStaticMeshComponents(const TArray<UStaticMeshComponent
 
 		FIntPoint AtlasTextureSize = InSettings.MaterialSettings.TextureSize;
 		FFlattenMaterial MergedFlatMaterial;
-		MergedFlatMaterial.DiffuseSize = AtlasTextureSize;
-		MergedFlatMaterial.NormalSize = InSettings.MaterialSettings.bNormalMap ? AtlasTextureSize : FIntPoint::ZeroValue;
-		MergedFlatMaterial.MetallicSize = InSettings.MaterialSettings.bMetallicMap ? AtlasTextureSize : FIntPoint::ZeroValue;
-		MergedFlatMaterial.RoughnessSize = InSettings.MaterialSettings.bRoughnessMap ? AtlasTextureSize : FIntPoint::ZeroValue;
-		MergedFlatMaterial.SpecularSize = InSettings.MaterialSettings.bSpecularMap ? AtlasTextureSize : FIntPoint::ZeroValue;
-		MergedFlatMaterial.EmissiveSize = InSettings.MaterialSettings.bEmissiveMap ? AtlasTextureSize : FIntPoint::ZeroValue;
-		MergedFlatMaterial.OpacitySize = InSettings.MaterialSettings.bOpacityMap ? AtlasTextureSize : FIntPoint::ZeroValue;
+		MergedFlatMaterial.SetPropertySize(EFlattenMaterialProperties::Diffuse, AtlasTextureSize);
+		MergedFlatMaterial.SetPropertySize(EFlattenMaterialProperties::Normal, InSettings.MaterialSettings.bNormalMap ? AtlasTextureSize : FIntPoint::ZeroValue);
+		MergedFlatMaterial.SetPropertySize(EFlattenMaterialProperties::Metallic, InSettings.MaterialSettings.bMetallicMap ? AtlasTextureSize : FIntPoint::ZeroValue);
+		MergedFlatMaterial.SetPropertySize(EFlattenMaterialProperties::Roughness, InSettings.MaterialSettings.bRoughnessMap ? AtlasTextureSize : FIntPoint::ZeroValue);
+		MergedFlatMaterial.SetPropertySize(EFlattenMaterialProperties::Specular, InSettings.MaterialSettings.bSpecularMap ? AtlasTextureSize : FIntPoint::ZeroValue);
+		MergedFlatMaterial.SetPropertySize(EFlattenMaterialProperties::Emissive, InSettings.MaterialSettings.bEmissiveMap ? AtlasTextureSize : FIntPoint::ZeroValue);
+		MergedFlatMaterial.SetPropertySize(EFlattenMaterialProperties::Opacity, InSettings.MaterialSettings.bOpacityMap ? AtlasTextureSize : FIntPoint::ZeroValue);
 
 		TArray<FRawMeshUVTransform> UVTransforms;
-		MergeFlattenedMaterials(FlattenedMaterials, MergedFlatMaterial, UVTransforms);
+
+		if (InSettings.bUseTextureBinning)
+		{
+			TArray<float> MaterialImportance;
+			FMaterialUtilities::DetermineMaterialImportance(UniqueMaterials, MaterialImportance);
+			TArray<FBox2D> MaterialBoxes;
+			FMaterialUtilities::GeneratedBinnedTextureSquares(FVector2D(1.0f, 1.0f), MaterialImportance, MaterialBoxes);
+			FlattenBinnedMaterials(FlattenedMaterials, MaterialBoxes, MergedFlatMaterial, UVTransforms);
+		}
+		else
+		{
+			MergeFlattenedMaterials(FlattenedMaterials, MergedFlatMaterial, UVTransforms);
+		}
 
 		FMaterialUtilities::OptimizeFlattenMaterial(MergedFlatMaterial);
 
@@ -7384,6 +7668,17 @@ void FMeshUtilities::MergeStaticMeshComponents(const TArray<UStaticMeshComponent
 	FRawMeshExt MergedMesh;
 	FMemory::Memset(&MergedMesh, 0, sizeof(MergedMesh));
 
+	// Flatten out the occupied UV channel flags, we need this to ensure the same amount of uv sets written out for each mesh
+	bool bFlattenedOcuppiedUVChannels[MAX_MESH_TEXTURE_COORDS];
+	FMemory::Memset(bFlattenedOcuppiedUVChannels, 0, sizeof(bool) * MAX_MESH_TEXTURE_COORDS);
+	for (int CoordinateIndex = 0; CoordinateIndex < MAX_MESH_TEXTURE_COORDS; ++CoordinateIndex)
+	{
+		for (int32 LODIndex = 0; LODIndex < MAX_STATIC_MESH_LODS; ++LODIndex)
+		{
+			bFlattenedOcuppiedUVChannels[CoordinateIndex] |= bOcuppiedUVChannels[LODIndex][CoordinateIndex];
+		}
+	}
+
 	int32 MaxExportLODs = bMergeAllAvailableLODs ? NumMaxLOD : 1;
 	// Merge meshes into single mesh
 	for (int32 SourceMeshIdx = 0; SourceMeshIdx < SourceMeshes.Num(); ++SourceMeshIdx)
@@ -7447,12 +7742,12 @@ void FMeshUtilities::MergeStaticMeshComponents(const TArray<UStaticMeshComponent
 
 			// Deal with vertex colors
 			// Some meshes may have it, in this case merged mesh will be forced to have vertex colors as well
-			if (bWithVertexColors[SourceLODIndex] && InSettings.bBakeVertexDataToMesh)
+			if (InSettings.bBakeVertexDataToMesh)
 			{
-				if (SourceRawMesh.WedgeColors.Num())
+				if (bWithVertexColors[SourceLODIndex] && SourceRawMesh.WedgeColors.Num())
 				{
 					TargetRawMesh.WedgeColors.Append(SourceRawMesh.WedgeColors);
-				}
+				}	
 				else
 				{
 					// In case this source mesh does not have vertex colors, fill target with 0xFF
@@ -7462,12 +7757,13 @@ void FMeshUtilities::MergeStaticMeshComponents(const TArray<UStaticMeshComponent
 					FMemory::Memset(&TargetRawMesh.WedgeColors[ColorsOffset], 0xFF, ColorsNum*TargetRawMesh.WedgeColors.GetTypeSize());
 				}
 			}
+			
 
 			// Merge all other UV channels 
 			for (int32 ChannelIdx = 0; ChannelIdx < MAX_MESH_TEXTURE_COORDS; ++ChannelIdx)
 			{
 				// Whether this channel has data
-				if (bOcuppiedUVChannels[SourceLODIndex][ChannelIdx])
+				if (bFlattenedOcuppiedUVChannels[ChannelIdx])
 				{
 					const TArray<FVector2D>& SourceChannel = SourceRawMesh.WedgeTexCoords[ChannelIdx];
 					TArray<FVector2D>& TargetChannel = TargetRawMesh.WedgeTexCoords[ChannelIdx];
@@ -7503,12 +7799,18 @@ void FMeshUtilities::MergeStaticMeshComponents(const TArray<UStaticMeshComponent
 	
 	// Compute target lightmap channel for each LOD
 	// User can specify any index, but there are should not be empty gaps in UV channel list
-	int32 TargetLightMapUVChannel[MAX_STATIC_MESH_LODS];
-	for (int32 LODIndex = 0; LODIndex < MAX_STATIC_MESH_LODS; ++LODIndex)
-	{		
-		for (int32 ChannelIdx = 0; ChannelIdx < MAX_MESH_TEXTURE_COORDS && bOcuppiedUVChannels[LODIndex][ChannelIdx]; ++ChannelIdx)
+	int32 LightMapUVChannel = 0;
+	for (int32 ChannelIdx = InSettings.TargetLightMapUVChannel; ChannelIdx < MAX_MESH_TEXTURE_COORDS; ++ChannelIdx)
+	{
+		bool bOccupied = false;
+		if (bFlattenedOcuppiedUVChannels[ChannelIdx])
+		{
+			continue;
+		}
+		else
 		{
-			TargetLightMapUVChannel[LODIndex] = FMath::Min(InSettings.TargetLightMapUVChannel, ChannelIdx);
+			LightMapUVChannel = ChannelIdx;
+			break;
 		}
 	}
 
@@ -7550,17 +7852,7 @@ void FMeshUtilities::MergeStaticMeshComponents(const TArray<UStaticMeshComponent
 		if (InSettings.bGenerateLightMapUV)
 		{
 			StaticMesh->LightMapResolution = InSettings.TargetLightMapResolution;
-
-			int32 TargetLightMapIndex = InSettings.TargetLightMapUVChannel;	
-			for (int32 LODIndex = 0; LODIndex < MAX_STATIC_MESH_LODS; ++LODIndex)
-			{
-				for (int32 ChannelIdx = 0; ChannelIdx < MAX_MESH_TEXTURE_COORDS && bOcuppiedUVChannels[LODIndex][ChannelIdx]; ++ChannelIdx)
-				{
-					TargetLightMapIndex = FMath::Max(TargetLightMapIndex, ChannelIdx);
-				}
-			}
-
-			StaticMesh->LightMapCoordinateIndex = TargetLightMapIndex + 1;
+			StaticMesh->LightMapCoordinateIndex = LightMapUVChannel;
 		}
 
 		for (int32 LODIndex = 0; LODIndex < NumMaxLOD; ++LODIndex)
@@ -7580,7 +7872,7 @@ void FMeshUtilities::MergeStaticMeshComponents(const TArray<UStaticMeshComponent
 					SrcModel->BuildSettings.bGenerateLightmapUVs = InSettings.bGenerateLightMapUV;
 					SrcModel->BuildSettings.MinLightmapResolution = InSettings.TargetLightMapResolution;
 					SrcModel->BuildSettings.SrcLightmapIndex = 0;
-					SrcModel->BuildSettings.DstLightmapIndex = StaticMesh->LightMapCoordinateIndex;
+					SrcModel->BuildSettings.DstLightmapIndex = LightMapUVChannel;
 
 					SrcModel->RawMeshBulkData->SaveRawMesh(MergedMeshLOD);
 				}
@@ -7788,10 +8080,44 @@ void FMeshUtilities::StartupModule()
 		bUsingNvTriStrip ? TEXT("_NvTriStrip") : TEXT("")
 		);
 	bUsingSimplygon = VersionString.Contains(TEXT("Simplygon"));
+
+	// hook up level editor extension for skeletal mesh conversion
+	ModuleLoadedDelegateHandle = FModuleManager::Get().OnModulesChanged().AddLambda([this](FName InModuleName, EModuleChangeReason InChangeReason)
+	{
+		if (InChangeReason == EModuleChangeReason::ModuleLoaded)
+		{
+			if (InModuleName == "LevelEditor")
+			{
+				AddLevelViewportMenuExtender();
+			}
+			else if (InModuleName == "AnimationBlueprintEditor")
+			{
+				AddAnimationBlueprintEditorToolbarExtender();
+			}
+			else if (InModuleName == "AnimationEditor")
+			{
+				AddAnimationEditorToolbarExtender();
+			}
+			else if (InModuleName == "SkeletalMeshEditor")
+			{
+				AddSkeletalMeshEditorToolbarExtender();
+			}
+			else if (InModuleName == "SkeletonEditor")
+			{
+				AddSkeletonEditorToolbarExtender();
+			}
+		}
+	});
 }
 
 void FMeshUtilities::ShutdownModule()
 {
+	RemoveLevelViewportMenuExtender();
+	RemoveAnimationBlueprintEditorToolbarExtender();
+	RemoveAnimationEditorToolbarExtender();
+	RemoveSkeletalMeshEditorToolbarExtender();
+	RemoveSkeletonEditorToolbarExtender();
+	FModuleManager::Get().OnModulesChanged().Remove(ModuleLoadedDelegateHandle);
 	MeshReduction = NULL;
 	MeshMerging = NULL;
 	VersionString.Empty();
@@ -7901,6 +8227,292 @@ void FMeshUtilities::CalculateTangents(const TArray<FVector>& InVertices, const
 	ComputeTangents(InVertices, InIndices, InUVs, InSmoothingGroupIndices, OverlappingCorners, OutTangentX, OutTangentY, OutNormals, InTangentOptions);
 }
 
+void FMeshUtilities::AddAnimationBlueprintEditorToolbarExtender()
+{
+	IAnimationBlueprintEditorModule& AnimationBlueprintEditorModule = FModuleManager::Get().LoadModuleChecked<IAnimationBlueprintEditorModule>("AnimationBlueprintEditor");
+	auto& ToolbarExtenders = AnimationBlueprintEditorModule.GetAllAnimationBlueprintEditorToolbarExtenders();
+
+	ToolbarExtenders.Add(IAnimationBlueprintEditorModule::FAnimationBlueprintEditorToolbarExtender::CreateRaw(this, &FMeshUtilities::GetAnimationBlueprintEditorToolbarExtender));
+	AnimationBlueprintEditorExtenderHandle = ToolbarExtenders.Last().GetHandle();
+}
+
+void FMeshUtilities::RemoveAnimationBlueprintEditorToolbarExtender()
+{
+	IAnimationBlueprintEditorModule* AnimationBlueprintEditorModule = FModuleManager::Get().GetModulePtr<IAnimationBlueprintEditorModule>("AnimationBlueprintEditor");
+	if (AnimationBlueprintEditorModule)
+	{
+		typedef IAnimationBlueprintEditorModule::FAnimationBlueprintEditorToolbarExtender DelegateType;
+		AnimationBlueprintEditorModule->GetAllAnimationBlueprintEditorToolbarExtenders().RemoveAll([=](const DelegateType& In) { return In.GetHandle() == AnimationBlueprintEditorExtenderHandle; });
+	}
+}
+
+TSharedRef<FExtender> FMeshUtilities::GetAnimationBlueprintEditorToolbarExtender(const TSharedRef<FUICommandList> CommandList, TSharedRef<IAnimationBlueprintEditor> InAnimationBlueprintEditor)
+{
+	TSharedRef<FExtender> Extender = MakeShareable(new FExtender);
+
+	UMeshComponent* MeshComponent = Cast<UMeshComponent>(InAnimationBlueprintEditor->GetPersonaToolkit()->GetPreviewMeshComponent());
+
+	Extender->AddToolBarExtension(
+		"Asset",
+		EExtensionHook::After,
+		CommandList,
+		FToolBarExtensionDelegate::CreateRaw(this, &FMeshUtilities::HandleAddConvertComponentToStaticMeshToToolbar, MeshComponent)
+	);
+
+	return Extender;
+}
+
+void FMeshUtilities::AddAnimationEditorToolbarExtender()
+{
+	IAnimationEditorModule& AnimationEditorModule = FModuleManager::Get().LoadModuleChecked<IAnimationEditorModule>("AnimationEditor");
+	auto& ToolbarExtenders = AnimationEditorModule.GetAllAnimationEditorToolbarExtenders();
+
+	ToolbarExtenders.Add(IAnimationEditorModule::FAnimationEditorToolbarExtender::CreateRaw(this, &FMeshUtilities::GetAnimationEditorToolbarExtender));
+	AnimationEditorExtenderHandle = ToolbarExtenders.Last().GetHandle();
+}
+
+void FMeshUtilities::RemoveAnimationEditorToolbarExtender()
+{
+	IAnimationEditorModule* AnimationEditorModule = FModuleManager::Get().GetModulePtr<IAnimationEditorModule>("AnimationEditor");
+	if (AnimationEditorModule)
+	{
+		typedef IAnimationEditorModule::FAnimationEditorToolbarExtender DelegateType;
+		AnimationEditorModule->GetAllAnimationEditorToolbarExtenders().RemoveAll([=](const DelegateType& In) { return In.GetHandle() == AnimationEditorExtenderHandle; });
+	}
+}
+
+TSharedRef<FExtender> FMeshUtilities::GetAnimationEditorToolbarExtender(const TSharedRef<FUICommandList> CommandList, TSharedRef<IAnimationEditor> InAnimationEditor)
+{
+	TSharedRef<FExtender> Extender = MakeShareable(new FExtender);
+
+	UMeshComponent* MeshComponent = Cast<UMeshComponent>(InAnimationEditor->GetPersonaToolkit()->GetPreviewMeshComponent());
+
+	Extender->AddToolBarExtension(
+		"Asset",
+		EExtensionHook::After,
+		CommandList,
+		FToolBarExtensionDelegate::CreateRaw(this, &FMeshUtilities::HandleAddConvertComponentToStaticMeshToToolbar, MeshComponent)
+	);
+
+	return Extender;
+}
+
+void FMeshUtilities::AddSkeletalMeshEditorToolbarExtender()
+{
+	ISkeletalMeshEditorModule& SkeletalMeshEditorModule = FModuleManager::Get().LoadModuleChecked<ISkeletalMeshEditorModule>("SkeletalMeshEditor");
+	auto& ToolbarExtenders = SkeletalMeshEditorModule.GetAllSkeletalMeshEditorToolbarExtenders();
+
+	ToolbarExtenders.Add(ISkeletalMeshEditorModule::FSkeletalMeshEditorToolbarExtender::CreateRaw(this, &FMeshUtilities::GetSkeletalMeshEditorToolbarExtender));
+	SkeletalMeshEditorExtenderHandle = ToolbarExtenders.Last().GetHandle();
+}
+
+void FMeshUtilities::RemoveSkeletalMeshEditorToolbarExtender()
+{
+	ISkeletalMeshEditorModule* SkeletalMeshEditorModule = FModuleManager::Get().GetModulePtr<ISkeletalMeshEditorModule>("SkeletalMeshEditor");
+	if (SkeletalMeshEditorModule)
+	{
+		typedef ISkeletalMeshEditorModule::FSkeletalMeshEditorToolbarExtender DelegateType;
+		SkeletalMeshEditorModule->GetAllSkeletalMeshEditorToolbarExtenders().RemoveAll([=](const DelegateType& In) { return In.GetHandle() == SkeletalMeshEditorExtenderHandle; });
+	}
+}
+
+TSharedRef<FExtender> FMeshUtilities::GetSkeletalMeshEditorToolbarExtender(const TSharedRef<FUICommandList> CommandList, TSharedRef<ISkeletalMeshEditor> InSkeletalMeshEditor)
+{
+	TSharedRef<FExtender> Extender = MakeShareable(new FExtender);
+
+	UMeshComponent* MeshComponent = Cast<UMeshComponent>(InSkeletalMeshEditor->GetPersonaToolkit()->GetPreviewMeshComponent());
+
+	Extender->AddToolBarExtension(
+		"Asset",
+		EExtensionHook::After,
+		CommandList,
+		FToolBarExtensionDelegate::CreateRaw(this, &FMeshUtilities::HandleAddConvertComponentToStaticMeshToToolbar, MeshComponent)
+	);
+
+	return Extender;
+}
+
+void FMeshUtilities::AddSkeletonEditorToolbarExtender()
+{
+	ISkeletonEditorModule& SkeletonEditorModule = FModuleManager::Get().LoadModuleChecked<ISkeletonEditorModule>("SkeletonEditor");
+	auto& ToolbarExtenders = SkeletonEditorModule.GetAllSkeletonEditorToolbarExtenders();
+
+	ToolbarExtenders.Add(ISkeletonEditorModule::FSkeletonEditorToolbarExtender::CreateRaw(this, &FMeshUtilities::GetSkeletonEditorToolbarExtender));
+	SkeletonEditorExtenderHandle = ToolbarExtenders.Last().GetHandle();
+}
+
+void FMeshUtilities::RemoveSkeletonEditorToolbarExtender()
+{
+	ISkeletonEditorModule* SkeletonEditorModule = FModuleManager::Get().GetModulePtr<ISkeletonEditorModule>("SkeletonEditor");
+	if (SkeletonEditorModule)
+	{
+		typedef ISkeletonEditorModule::FSkeletonEditorToolbarExtender DelegateType;
+		SkeletonEditorModule->GetAllSkeletonEditorToolbarExtenders().RemoveAll([=](const DelegateType& In) { return In.GetHandle() == SkeletonEditorExtenderHandle; });
+	}
+}
+
+TSharedRef<FExtender> FMeshUtilities::GetSkeletonEditorToolbarExtender(const TSharedRef<FUICommandList> CommandList, TSharedRef<ISkeletonEditor> InSkeletonEditor)
+{
+	TSharedRef<FExtender> Extender = MakeShareable(new FExtender);
+
+	UMeshComponent* MeshComponent = Cast<UMeshComponent>(InSkeletonEditor->GetPersonaToolkit()->GetPreviewMeshComponent());
+
+	Extender->AddToolBarExtension(
+		"Asset",
+		EExtensionHook::After,
+		CommandList,
+		FToolBarExtensionDelegate::CreateRaw(this, &FMeshUtilities::HandleAddConvertComponentToStaticMeshToToolbar, MeshComponent)
+	);
+
+	return Extender;
+}
+
+
+void FMeshUtilities::HandleAddConvertComponentToStaticMeshToToolbar(FToolBarBuilder& ParentToolbarBuilder, UMeshComponent* InMeshComponent)
+{
+	ParentToolbarBuilder.AddToolBarButton(
+		FUIAction(FExecuteAction::CreateLambda([this, InMeshComponent]()
+		{
+			ConvertMeshesToStaticMesh(TArray<UMeshComponent*>({ InMeshComponent }), InMeshComponent->GetComponentToWorld());
+		})),
+		NAME_None,
+		LOCTEXT("MakeStaticMesh", "Make Static Mesh"),
+		LOCTEXT("MakeStaticMeshTooltip", "Make a new static mesh out of the preview's current pose."),
+		FSlateIcon("EditorStyle", "Persona.ConvertToStaticMesh")
+	);
+}
+
+void FMeshUtilities::AddLevelViewportMenuExtender()
+{
+	FLevelEditorModule& LevelEditorModule = FModuleManager::Get().LoadModuleChecked<FLevelEditorModule>("LevelEditor");
+	auto& MenuExtenders = LevelEditorModule.GetAllLevelViewportContextMenuExtenders();
+
+	MenuExtenders.Add(FLevelEditorModule::FLevelViewportMenuExtender_SelectedActors::CreateRaw(this, &FMeshUtilities::GetLevelViewportContextMenuExtender));
+	LevelViewportExtenderHandle = MenuExtenders.Last().GetHandle();
+}
+
+void FMeshUtilities::RemoveLevelViewportMenuExtender()
+{
+	if (LevelViewportExtenderHandle.IsValid())
+	{
+		FLevelEditorModule* LevelEditorModule = FModuleManager::Get().GetModulePtr<FLevelEditorModule>("LevelEditor");
+		if (LevelEditorModule)
+		{
+			typedef FLevelEditorModule::FLevelViewportMenuExtender_SelectedActors DelegateType;
+			LevelEditorModule->GetAllLevelViewportContextMenuExtenders().RemoveAll([=](const DelegateType& In) { return In.GetHandle() == LevelViewportExtenderHandle; });
+		}
+	}
+}
+
+TSharedRef<FExtender> FMeshUtilities::GetLevelViewportContextMenuExtender(const TSharedRef<FUICommandList> CommandList, const TArray<AActor*> InActors)
+{
+	TSharedRef<FExtender> Extender = MakeShareable(new FExtender);
+
+	if (InActors.Num() > 0)
+	{
+		bool bHasComponent = false;
+		for (AActor* Actor : InActors)
+		{
+			TInlineComponentArray<UMeshComponent*> Components(Actor);
+			if (Components.Num() > 0)
+			{
+				bHasComponent = true;
+				break;
+			}
+		}
+		
+		if (bHasComponent)
+		{
+			FText ActorName = InActors.Num() == 1 ? FText::Format(LOCTEXT("ActorNameSingular", "\"{0}\""), FText::FromString(InActors[0]->GetActorLabel())) : LOCTEXT("ActorNamePlural", "Actors");
+
+			FLevelEditorModule& LevelEditor = FModuleManager::GetModuleChecked<FLevelEditorModule>(TEXT("LevelEditor"));
+			TSharedRef<FUICommandList> LevelEditorCommandBindings = LevelEditor.GetGlobalLevelEditorActions();
+
+			Extender->AddMenuExtension("ActorControl", EExtensionHook::After, LevelEditorCommandBindings, FMenuExtensionDelegate::CreateLambda(
+				[this, ActorName, InActors](FMenuBuilder& MenuBuilder) {
+
+				MenuBuilder.AddMenuEntry(
+					FText::Format(LOCTEXT("ConvertSelectedActorsToStaticMeshText", "Convert {0} To Static Mesh"), ActorName),
+					LOCTEXT("ConvertSelectedActorsToStaticMeshTooltip", "Convert the selected actor's meshes to a new Static Mesh asset. Supports static and skeletal meshes."),
+					FSlateIcon(),
+					FUIAction(FExecuteAction::CreateRaw(this, &FMeshUtilities::ConvertActorMeshesToStaticMesh, InActors))
+				);
+			})
+			);
+		}
+	}
+
+	return Extender;
+}
+
+void FMeshUtilities::ConvertActorMeshesToStaticMesh(const TArray<AActor*> InActors)
+{
+	TArray<UMeshComponent*> MeshComponents;
+
+	for (AActor* Actor : InActors)
+	{
+		// add all components from this actor
+		TInlineComponentArray<UMeshComponent*> ActorComponents(Actor);
+		for (UMeshComponent* ActorComponent : ActorComponents)
+		{
+			MeshComponents.AddUnique(ActorComponent);
+		}
+
+		// add all attached actors
+		TArray<AActor*> AttachedActors;
+		Actor->GetAttachedActors(AttachedActors);
+		for (AActor* AttachedActor : AttachedActors)
+		{
+			TInlineComponentArray<UMeshComponent*> AttachedActorComponents(AttachedActor);
+			for (UMeshComponent* AttachedActorComponent : AttachedActorComponents)
+			{
+				MeshComponents.AddUnique(AttachedActorComponent);
+			}
+		}
+	}
+
+	auto GetActorRootTransform = [](AActor* InActor)
+	{
+		FTransform RootTransform(FTransform::Identity);
+		if (ACharacter* Character = Cast<ACharacter>(InActor))
+		{
+			RootTransform = Character->GetTransform();
+			RootTransform.SetLocation(RootTransform.GetLocation() - FVector(0.0f, 0.0f, Character->GetCapsuleComponent()->GetScaledCapsuleHalfHeight()));
+		}
+		else
+		{
+			// otherwise just use the actor's origin
+			RootTransform = InActor->GetTransform();
+		}
+
+		return RootTransform;
+	};
+
+	// now pick a root transform
+	FTransform RootTransform(FTransform::Identity);
+	if (InActors.Num() == 1)
+	{
+		RootTransform = GetActorRootTransform(InActors[0]);
+	}
+	else
+	{
+		// multiple actors use the average of their origins, with Z being the min of all origins. Rotation is identity for simplicity
+		FVector Location(FVector::ZeroVector);
+		float MinZ = FLT_MAX;
+		for (AActor* Actor : InActors)
+		{
+			FTransform ActorTransform(GetActorRootTransform(Actor));
+			Location += ActorTransform.GetLocation();
+			MinZ = FMath::Min(ActorTransform.GetLocation().Z, MinZ);
+		}
+		Location /= (float)InActors.Num();
+		Location.Z = MinZ;
+
+		RootTransform.SetLocation(Location);
+	}
+
+	ConvertMeshesToStaticMesh(MeshComponents, RootTransform);
+}
 
 #undef LOCTEXT_NAMESPACE
 
diff --git a/Engine/Source/Developer/MeshUtilities/Private/MeshUtilitiesPrivate.h b/Engine/Source/Developer/MeshUtilities/Private/MeshUtilitiesPrivate.h
--- a/Engine/Source/Developer/MeshUtilities/Private/MeshUtilitiesPrivate.h
+++ b/Engine/Source/Developer/MeshUtilities/Private/MeshUtilitiesPrivate.h
@@ -5,3 +5,4 @@
 #include "Engine.h"
 #include "MeshUtilities.h"
 #include "RawMesh.h"
+#include "SlateBasics.h"
diff --git a/Engine/Source/Developer/MeshUtilities/Private/ProxyMaterialUtilities.h b/Engine/Source/Developer/MeshUtilities/Private/ProxyMaterialUtilities.h
--- a/Engine/Source/Developer/MeshUtilities/Private/ProxyMaterialUtilities.h
+++ b/Engine/Source/Developer/MeshUtilities/Private/ProxyMaterialUtilities.h
@@ -5,10 +5,10 @@
 namespace ProxyMaterialUtilities
 {
 #define TEXTURE_MACRO_BASE(a, b, c) \
-	bool b##a##Texture = FlattenMaterial.a##Samples.Num() > 1; \
+	bool b##a##Texture = FlattenMaterial.DoesPropertyContainData(EFlattenMaterialProperties::a) && !FlattenMaterial.IsPropertyConstant(EFlattenMaterialProperties::a); \
 	if ( b##a##Texture) \
 	{ \
-		UTexture2D* a##Texture = b##a##Texture ? FMaterialUtilities::CreateTexture(InOuter, AssetBasePath + TEXT("T_") + AssetBaseName + TEXT("_" #a), FlattenMaterial.a##Size, FlattenMaterial.a##Samples, b, TEXTUREGROUP_HierarchicalLOD, RF_Public | RF_Standalone, c) : nullptr; \
+		UTexture2D* a##Texture = b##a##Texture ? FMaterialUtilities::CreateTexture(InOuter, AssetBasePath + TEXT("T_") + AssetBaseName + TEXT("_" #a), FlattenMaterial.GetPropertySize(EFlattenMaterialProperties::a), FlattenMaterial.GetPropertySamples(EFlattenMaterialProperties::a), b, TEXTUREGROUP_HierarchicalLOD, RF_Public | RF_Standalone, c) : nullptr; \
 		OutMaterial->SetTextureParameterValueEditorOnly(#a "Texture", a##Texture); \
 		FStaticSwitchParameter SwitchParameter; \
 		SwitchParameter.ParameterName = "Use" #a; \
@@ -21,46 +21,52 @@ namespace ProxyMaterialUtilities
 #define TEXTURE_MACRO_VECTOR(a, b, c) TEXTURE_MACRO_BASE(a, b, c)\
 	else\
 	{ \
-		OutMaterial->SetVectorParameterValueEditorOnly(#a "Const", FlattenMaterial.a##Samples[0]); \
+		OutMaterial->SetVectorParameterValueEditorOnly(#a "Const", FlattenMaterial.GetPropertySamples(EFlattenMaterialProperties::a)[0]); \
 	} 
 
 #define TEXTURE_MACRO_VECTOR_LINEAR(a, b, c) TEXTURE_MACRO_BASE(a, b, c)\
 	else\
 	{ \
-		OutMaterial->SetVectorParameterValueEditorOnly(#a "Const", FlattenMaterial.a##Samples[0].ReinterpretAsLinear()); \
+		OutMaterial->SetVectorParameterValueEditorOnly(#a "Const", FlattenMaterial.GetPropertySamples(EFlattenMaterialProperties::a)[0].ReinterpretAsLinear()); \
 	} 
 
 #define TEXTURE_MACRO_SCALAR(a, b, c) TEXTURE_MACRO_BASE(a, b, c)\
 	else \
 	{ \
-		FLinearColor Colour = ( FlattenMaterial.a##Samples.Num() == 1 ) ? FLinearColor::FromSRGBColor( FlattenMaterial.a##Samples[0]) : FLinearColor( InMaterialProxySettings.a##Constant, 0, 0, 0 ); \
+		FLinearColor Colour = FlattenMaterial.IsPropertyConstant(EFlattenMaterialProperties::a) ? FLinearColor::FromSRGBColor( FlattenMaterial.GetPropertySamples(EFlattenMaterialProperties::a)[0]) : FLinearColor( InMaterialProxySettings.a##Constant, 0, 0, 0 ); \
 		OutMaterial->SetScalarParameterValueEditorOnly(#a "Const", Colour.R ); \
 	}
 
 	static const bool CalculatePackedTextureData(const FFlattenMaterial& InMaterial, bool& bOutPackMetallic, bool& bOutPackSpecular, bool& bOutPackRoughness, int32& OutNumSamples, FIntPoint& OutSize)
 	{
 		// Whether or not a material property is baked down
-		const bool bHasMetallic = (InMaterial.MetallicSamples.Num() > 1);
-		const bool bHasRoughness = (InMaterial.RoughnessSamples.Num() > 1);
-		const bool bHasSpecular = (InMaterial.SpecularSamples.Num() > 1);
+		const bool bHasMetallic = InMaterial.DoesPropertyContainData(EFlattenMaterialProperties::Metallic) && !InMaterial.IsPropertyConstant(EFlattenMaterialProperties::Metallic);
+		const bool bHasRoughness = InMaterial.DoesPropertyContainData(EFlattenMaterialProperties::Roughness) && !InMaterial.IsPropertyConstant(EFlattenMaterialProperties::Roughness);
+		const bool bHasSpecular = InMaterial.DoesPropertyContainData(EFlattenMaterialProperties::Specular) && !InMaterial.IsPropertyConstant(EFlattenMaterialProperties::Specular);
 
 		// Check for same texture sizes
 		bool bSameTextureSize = false;
 
 		// Determine whether or not the properties sizes match
-		bSameTextureSize = (InMaterial.MetallicSize == InMaterial.RoughnessSize) || (InMaterial.MetallicSize == InMaterial.SpecularSize);
+
+
+		const FIntPoint MetallicSize = InMaterial.GetPropertySize(EFlattenMaterialProperties::Metallic);
+		const FIntPoint SpecularSize = InMaterial.GetPropertySize(EFlattenMaterialProperties::Specular);
+		const FIntPoint RoughnessSize = InMaterial.GetPropertySize(EFlattenMaterialProperties::Roughness);
+
+		bSameTextureSize = (MetallicSize == RoughnessSize) || (MetallicSize == SpecularSize);
 		if (bSameTextureSize)
 		{
-			OutSize = InMaterial.MetallicSize;
-			OutNumSamples = InMaterial.MetallicSamples.Num();
+			OutSize = MetallicSize;
+			OutNumSamples = InMaterial.GetPropertySamples(EFlattenMaterialProperties::Metallic).Num();
 		}
 		else
 		{
-			bSameTextureSize = (InMaterial.RoughnessSize == InMaterial.SpecularSize);
+			bSameTextureSize = (RoughnessSize == SpecularSize);
 			if (bSameTextureSize)
 			{
-				OutSize = InMaterial.RoughnessSize;
-				OutNumSamples = InMaterial.RoughnessSamples.Num();
+				OutSize = RoughnessSize;
+				OutNumSamples = InMaterial.GetPropertySamples(EFlattenMaterialProperties::Roughness).Num();
 			}
 		}
 
@@ -68,11 +74,11 @@ namespace ProxyMaterialUtilities
 		int32 NumPacked = 0;
 		if (OutNumSamples != 0)
 		{
-			bOutPackMetallic = bHasMetallic ? (OutNumSamples == InMaterial.MetallicSamples.Num()) : false;
+			bOutPackMetallic = bHasMetallic ? (OutNumSamples == InMaterial.GetPropertySamples(EFlattenMaterialProperties::Metallic).Num()) : false;
 			NumPacked += (bOutPackMetallic) ? 1 : 0;
-			bOutPackRoughness = bHasRoughness ? (OutNumSamples == InMaterial.RoughnessSamples.Num()) : false;
+			bOutPackRoughness = bHasRoughness ? (OutNumSamples == InMaterial.GetPropertySamples(EFlattenMaterialProperties::Roughness).Num()) : false;
 			NumPacked += (bOutPackRoughness) ? 1 : 0;
-			bOutPackSpecular = bHasSpecular ? (OutNumSamples == InMaterial.SpecularSamples.Num()) : false;
+			bOutPackSpecular = bHasSpecular ? (OutNumSamples == InMaterial.GetPropertySamples(EFlattenMaterialProperties::Specular).Num()) : false;
 			NumPacked += (bOutPackSpecular) ? 1 : 0;
 		}
 		else
@@ -107,28 +113,28 @@ namespace ProxyMaterialUtilities
 		FStaticParameterSet NewStaticParameterSet;
 
 		// Load textures and set switches accordingly
-		if (FlattenMaterial.DiffuseSamples.Num() > 0 && !(FlattenMaterial.DiffuseSamples.Num() == 1 && FlattenMaterial.DiffuseSamples[0] == FColor::Black))
+		if (FlattenMaterial.GetPropertySize(EFlattenMaterialProperties::Diffuse).Num() > 0 && !(FlattenMaterial.IsPropertyConstant(EFlattenMaterialProperties::Diffuse) && FlattenMaterial.GetPropertySamples(EFlattenMaterialProperties::Diffuse)[0] == FColor::Black))
 		{
 			TEXTURE_MACRO_VECTOR(Diffuse, TC_Default, bSRGB);
 		}
 
-		if (FlattenMaterial.NormalSamples.Num() > 1)
+		if (FlattenMaterial.GetPropertySize(EFlattenMaterialProperties::Normal).Num() > 1)
 		{
 			TEXTURE_MACRO_BASE(Normal, TC_Normalmap, bRGB)
 		}
 
 		// Determine whether or not specific material properties are packed together into one texture (requires at least two to match (number of samples and texture size) in order to be packed
-		if (!bPackMetallic && (FlattenMaterial.MetallicSamples.Num() > 0 || !InMaterialProxySettings.bMetallicMap))
+		if (!bPackMetallic && (FlattenMaterial.GetPropertySize(EFlattenMaterialProperties::Metallic).Num() > 0 || !InMaterialProxySettings.bMetallicMap))
 		{
 			TEXTURE_MACRO_SCALAR(Metallic, TC_Default, bSRGB);
 		}
 
-		if (!bPackRoughness && (FlattenMaterial.RoughnessSamples.Num() > 0 || !InMaterialProxySettings.bRoughnessMap))
+		if (!bPackRoughness && (FlattenMaterial.GetPropertySize(EFlattenMaterialProperties::Roughness).Num() > 0 || !InMaterialProxySettings.bRoughnessMap))
 		{
 			TEXTURE_MACRO_SCALAR(Roughness, TC_Default, bSRGB);
 		}
 
-		if (!bPackSpecular && (FlattenMaterial.SpecularSamples.Num() > 0 || !InMaterialProxySettings.bSpecularMap))
+		if (!bPackSpecular && (FlattenMaterial.GetPropertySize(EFlattenMaterialProperties::Specular).Num() > 0 || !InMaterialProxySettings.bSpecularMap))
 		{
 			TEXTURE_MACRO_SCALAR(Specular, TC_Default, bSRGB);
 		}
@@ -140,11 +146,28 @@ namespace ProxyMaterialUtilities
 			MergedTexture.AddZeroed(NumSamples);
 
 			// Merge properties into one texture using the separate colour channels
-			if (bPackMetallic) for (int32 SampleIndex = 0; SampleIndex < NumSamples; ++SampleIndex) { MergedTexture[SampleIndex].R = FlattenMaterial.MetallicSamples[SampleIndex].R; }
 
-			if (bPackRoughness) for (int32 SampleIndex = 0; SampleIndex < NumSamples; ++SampleIndex) { MergedTexture[SampleIndex].G = FlattenMaterial.RoughnessSamples[SampleIndex].R; }
-
-			if (bPackSpecular) for (int32 SampleIndex = 0; SampleIndex < NumSamples; ++SampleIndex) { MergedTexture[SampleIndex].B = FlattenMaterial.SpecularSamples[SampleIndex].R; }
+			// R G B masks
+#if PLATFORM_LITTLE_ENDIAN
+			const uint32 ColorMask[3] = { 0x0000FF00, 0x00FF0000, 0xFF000000 };
+#else // PLATFORM_LITTLE_ENDIAN
+			const uint32 ColorMask[3] = { 0x00FF0000, 0x0000FF00, 0x000000FF };
+#endif
+			for (int32 PropertyIndex = 0; PropertyIndex < 3; ++PropertyIndex)
+			{
+				EFlattenMaterialProperties Property = (EFlattenMaterialProperties)(PropertyIndex + (int32)EFlattenMaterialProperties::Metallic);
+				const bool HasProperty = FlattenMaterial.DoesPropertyContainData(Property) && !FlattenMaterial.IsPropertyConstant(Property);
+
+				if (HasProperty)
+				{
+					const TArray<FColor>& PropertySamples = FlattenMaterial.GetPropertySamples(Property);
+					// OR masked values (samples initialized to zero, so no random data)
+					for (int32 SampleIndex = 0; SampleIndex < NumSamples; ++SampleIndex)
+					{
+						MergedTexture[SampleIndex].DWColor() |= (PropertySamples[SampleIndex].DWColor() & ColorMask[PropertyIndex]);
+					}
+				}
+			}
 
 			// Create texture using the merged property data
 			UTexture2D* PackedTexture = FMaterialUtilities::CreateTexture(InOuter, AssetBasePath + TEXT("T_") + AssetBaseName + TEXT("_MRS"), PackedSize, MergedTexture, TC_Default, TEXTUREGROUP_HierarchicalLOD, RF_Public | RF_Standalone, bSRGB);
@@ -172,7 +195,7 @@ namespace ProxyMaterialUtilities
 		}
 
 		// Emissive is a special case due to the scaling variable
-		if (FlattenMaterial.EmissiveSamples.Num() > 0 && !(FlattenMaterial.EmissiveSamples.Num() == 1 && FlattenMaterial.EmissiveSamples[0] == FColor::Black))
+		if (FlattenMaterial.GetPropertySamples(EFlattenMaterialProperties::Emissive).Num() > 0 && !(FlattenMaterial.GetPropertySize(EFlattenMaterialProperties::Emissive).Num() == 1 && FlattenMaterial.GetPropertySamples(EFlattenMaterialProperties::Emissive)[0] == FColor::Black))
 		{
 			TEXTURE_MACRO_VECTOR_LINEAR(Emissive, TC_Default, bRGB);
 
diff --git a/Engine/Source/Developer/MeshUtilities/Public/MeshUtilities.h b/Engine/Source/Developer/MeshUtilities/Public/MeshUtilities.h
--- a/Engine/Source/Developer/MeshUtilities/Public/MeshUtilities.h
+++ b/Engine/Source/Developer/MeshUtilities/Public/MeshUtilities.h
@@ -10,6 +10,8 @@ typedef FIntPoint FMeshIdAndLOD;
 struct FFlattenMaterial;
 struct FReferenceSkeleton;
 struct FStaticMeshLODResources;
+class UMeshComponent;
+class UStaticMesh;
 
 namespace ETangentOptions
 {
@@ -230,6 +232,15 @@ public:
 	 */
 	virtual void CalcBoneVertInfos( USkeletalMesh* SkeletalMesh, TArray<FBoneVertInfo>& Infos, bool bOnlyDominant) = 0;
 
+	/** 
+	 * Convert a set of mesh components in their current pose to a static mesh. 
+	 * @param	InMeshComponents		The mesh components we want to convert
+	 * @param	InRootTransform			The transform of the root of the mesh we want to output
+	 * @param	InPackageName			The package name to create the static mesh in. If this is empty then a dialog will be displayed to pick the mesh.
+	 * @return a new static mesh (specified by the user)
+	 */
+	virtual UStaticMesh* ConvertMeshesToStaticMesh(const TArray<UMeshComponent*>& InMeshComponents, const FTransform& InRootTransform = FTransform::Identity, const FString& InPackageName = FString()) = 0;
+
 	/**
 	 * Harvest static mesh components from input actors 
 	 * and merge into signle mesh grouping them by unique materials
diff --git a/Engine/Source/Developer/OutputLog/Private/SOutputLog.cpp b/Engine/Source/Developer/OutputLog/Private/SOutputLog.cpp
--- a/Engine/Source/Developer/OutputLog/Private/SOutputLog.cpp
+++ b/Engine/Source/Developer/OutputLog/Private/SOutputLog.cpp
@@ -3,9 +3,9 @@
 #include "OutputLogPrivatePCH.h"
 #include "SOutputLog.h"
 #include "SScrollBorder.h"
-#include "GameFramework/GameMode.h"
+#include "GameFramework/GameModeBase.h"
+#include "GameFramework/GameStateBase.h"
 #include "Engine/LocalPlayer.h"
-#include "GameFramework/GameState.h"
 #include "SSearchBox.h"
 
 #define LOCTEXT_NAMESPACE "SOutputLog"
@@ -387,12 +387,13 @@ void SConsoleInputBox::OnTextCommitted( const FText& InText, ETextCommit::Type C
 				{
 					if( !bWasHandled )
 					{					
-						AGameMode* const GameMode = World->GetAuthGameMode();
+						AGameModeBase* const GameMode = World->GetAuthGameMode();
+						AGameStateBase* const GameState = World->GetGameState();
 						if( GameMode && GameMode->ProcessConsoleExec( *ExecString, *GLog, NULL ) )
 						{
 							bWasHandled = true;
 						}
-						else if (World->GameState && World->GameState->ProcessConsoleExec(*ExecString, *GLog, NULL))
+						else if (GameState && GameState->ProcessConsoleExec(*ExecString, *GLog, NULL))
 						{
 							bWasHandled = true;
 						}
diff --git a/Engine/Source/Developer/SimplygonMeshReduction/Private/SimplygonMeshReduction.cpp b/Engine/Source/Developer/SimplygonMeshReduction/Private/SimplygonMeshReduction.cpp
--- a/Engine/Source/Developer/SimplygonMeshReduction/Private/SimplygonMeshReduction.cpp
+++ b/Engine/Source/Developer/SimplygonMeshReduction/Private/SimplygonMeshReduction.cpp
@@ -2156,54 +2156,83 @@ private:
 		SimplygonSDK::spMaterialTable& InSGMaterialTable,
 		FFlattenMaterial& OutMaterial)
 	{
+		FIntPoint Size;
 		SimplygonSDK::spMaterial SGMaterial = InSGMaterialTable->GetMaterial(0);
 
-		// Diffuse
-		GetMaterialChannelData(SGMaterial, SimplygonSDK::SG_MATERIAL_CHANNEL_DIFFUSE, OutMaterial.DiffuseSamples, OutMaterial.DiffuseSize);
+		// Diffuse		
+		Size = FIntPoint::ZeroValue;
+		GetMaterialChannelData(SGMaterial, SimplygonSDK::SG_MATERIAL_CHANNEL_DIFFUSE, OutMaterial.GetPropertySamples(EFlattenMaterialProperties::Diffuse), Size);
+		OutMaterial.SetPropertySize(EFlattenMaterialProperties::Diffuse, Size);
 
 		// Normal
-		GetMaterialChannelData(SGMaterial, SimplygonSDK::SG_MATERIAL_CHANNEL_NORMALS, OutMaterial.NormalSamples, OutMaterial.NormalSize);
+		Size = FIntPoint::ZeroValue;
+		GetMaterialChannelData(SGMaterial, SimplygonSDK::SG_MATERIAL_CHANNEL_NORMALS, OutMaterial.GetPropertySamples(EFlattenMaterialProperties::Normal), Size);
+		OutMaterial.SetPropertySize(EFlattenMaterialProperties::Normal, Size);
 
 		// Metallic
-		GetMaterialChannelData(SGMaterial, USER_MATERIAL_CHANNEL_METALLIC, OutMaterial.MetallicSamples, OutMaterial.MetallicSize);
+		Size = FIntPoint::ZeroValue;
+		GetMaterialChannelData(SGMaterial, USER_MATERIAL_CHANNEL_METALLIC, OutMaterial.GetPropertySamples(EFlattenMaterialProperties::Metallic), Size);
+		OutMaterial.SetPropertySize(EFlattenMaterialProperties::Metallic, Size);
 
 		// Roughness
-		GetMaterialChannelData(SGMaterial, USER_MATERIAL_CHANNEL_ROUGHNESS, OutMaterial.RoughnessSamples, OutMaterial.RoughnessSize);
+		Size = FIntPoint::ZeroValue;
+		GetMaterialChannelData(SGMaterial, USER_MATERIAL_CHANNEL_ROUGHNESS, OutMaterial.GetPropertySamples(EFlattenMaterialProperties::Roughness), Size);
+		OutMaterial.SetPropertySize(EFlattenMaterialProperties::Roughness, Size);
 
 		// Specular
-		GetMaterialChannelData(SGMaterial, USER_MATERIAL_CHANNEL_SPECULAR, OutMaterial.SpecularSamples, OutMaterial.SpecularSize);
+		Size = FIntPoint::ZeroValue;
+		GetMaterialChannelData(SGMaterial, USER_MATERIAL_CHANNEL_SPECULAR, OutMaterial.GetPropertySamples(EFlattenMaterialProperties::Specular), Size);
+		OutMaterial.SetPropertySize(EFlattenMaterialProperties::Specular, Size);
 	}
 
 	void CreateFlattenMaterialFromSGMaterial(
 		SimplygonSDK::spMaterial& SGMaterial,
 		FFlattenMaterial& OutMaterial)
 	{
+		FIntPoint Size;
+
 		// Diffuse
-		GetMaterialChannelData(SGMaterial, SimplygonSDK::SG_MATERIAL_CHANNEL_BASECOLOR, OutMaterial.DiffuseSamples, OutMaterial.DiffuseSize);
+		Size = FIntPoint::ZeroValue;
+		GetMaterialChannelData(SGMaterial, SimplygonSDK::SG_MATERIAL_CHANNEL_BASECOLOR, OutMaterial.GetPropertySamples(EFlattenMaterialProperties::Diffuse), Size);
+		OutMaterial.SetPropertySize(EFlattenMaterialProperties::Diffuse, Size);
 
 		// Normal
-		GetMaterialChannelData(SGMaterial, SimplygonSDK::SG_MATERIAL_CHANNEL_NORMALS, OutMaterial.NormalSamples, OutMaterial.NormalSize);
+		Size = FIntPoint::ZeroValue;
+		GetMaterialChannelData(SGMaterial, SimplygonSDK::SG_MATERIAL_CHANNEL_NORMALS, OutMaterial.GetPropertySamples(EFlattenMaterialProperties::Normal), Size);
+		OutMaterial.SetPropertySize(EFlattenMaterialProperties::Normal, Size);
 
 		// Metallic
-		GetMaterialChannelData(SGMaterial, SimplygonSDK::SG_MATERIAL_CHANNEL_METALLIC, OutMaterial.MetallicSamples, OutMaterial.MetallicSize);
+		Size = FIntPoint::ZeroValue;
+		GetMaterialChannelData(SGMaterial, SimplygonSDK::SG_MATERIAL_CHANNEL_METALLIC, OutMaterial.GetPropertySamples(EFlattenMaterialProperties::Metallic), Size);
+		OutMaterial.SetPropertySize(EFlattenMaterialProperties::Metallic, Size);
 
 		// Roughness
-		GetMaterialChannelData(SGMaterial, SimplygonSDK::SG_MATERIAL_CHANNEL_ROUGHNESS, OutMaterial.RoughnessSamples, OutMaterial.RoughnessSize);
+		Size = FIntPoint::ZeroValue;
+		GetMaterialChannelData(SGMaterial, SimplygonSDK::SG_MATERIAL_CHANNEL_ROUGHNESS, OutMaterial.GetPropertySamples(EFlattenMaterialProperties::Roughness), Size);
+		OutMaterial.SetPropertySize(EFlattenMaterialProperties::Roughness, Size);
 
 		// Specular
-		GetMaterialChannelData(SGMaterial, SimplygonSDK::SG_MATERIAL_CHANNEL_SPECULAR, OutMaterial.SpecularSamples, OutMaterial.SpecularSize);
+		Size = FIntPoint::ZeroValue;
+		GetMaterialChannelData(SGMaterial, SimplygonSDK::SG_MATERIAL_CHANNEL_SPECULAR, OutMaterial.GetPropertySamples(EFlattenMaterialProperties::Specular), Size);
+		OutMaterial.SetPropertySize(EFlattenMaterialProperties::Specular, Size);
 
 		// Opacity
+		Size = FIntPoint::ZeroValue;
 #if USE_USER_OPACITY_CHANNEL
-		GetMaterialChannelData(SGMaterial, USER_MATERIAL_CHANNEL_OPACITY, OutMaterial.OpacitySamples, OutMaterial.OpacitySize);
+		GetMaterialChannelData(SGMaterial, USER_MATERIAL_CHANNEL_OPACITY, OutMaterial.GetPropertySamples(EFlattenMaterialProperties::Opacity), Size);
 #else
-		GetMaterialChannelData(SGMaterial, SimplygonSDK::SG_MATERIAL_CHANNEL_OPACITY, OutMaterial.OpacitySamples, OutMaterial.OpacitySize);
+		GetMaterialChannelData(SGMaterial, SimplygonSDK::SG_MATERIAL_CHANNEL_OPACITY, OutMaterial.GetPropertySamples(EFlattenMaterialProperties::Opacity), Size);		
 #endif
+		OutMaterial.SetPropertySize(EFlattenMaterialProperties::Opacity, Size);
 
 		// Emissive
-		GetMaterialChannelData(SGMaterial, SimplygonSDK::SG_MATERIAL_CHANNEL_EMISSIVE, OutMaterial.EmissiveSamples, OutMaterial.EmissiveSize);
+		Size = FIntPoint::ZeroValue;
+		GetMaterialChannelData(SGMaterial, SimplygonSDK::SG_MATERIAL_CHANNEL_EMISSIVE, OutMaterial.GetPropertySamples(EFlattenMaterialProperties::Emissive), Size);
+		OutMaterial.SetPropertySize(EFlattenMaterialProperties::Emissive, Size);
 		
-		GetMaterialChannelData(SGMaterial, USER_MATERIAL_CHANNEL_SUBSURFACE_COLOR, OutMaterial.SubSurfaceSamples, OutMaterial.SubSurfaceSize);
+		Size = FIntPoint::ZeroValue;
+		GetMaterialChannelData(SGMaterial, USER_MATERIAL_CHANNEL_SUBSURFACE_COLOR, OutMaterial.GetPropertySamples(EFlattenMaterialProperties::SubSurface), Size);
+		OutMaterial.SetPropertySize(EFlattenMaterialProperties::SubSurface, Size);
 	}
 
 	static FIntPoint ComputeMappingImageSize(const FMaterialSimplificationSettings& Settings)
@@ -2241,55 +2270,55 @@ private:
 			SGMaterial->AddUserChannel(USER_MATERIAL_CHANNEL_SUBSURFACE_COLOR);
 
 			SGMaterial->SetName(TCHAR_TO_ANSI(*FString::Printf(TEXT("Material%d"), MaterialIndex)));
-
+			
 			// Does current material have BaseColor?
-			if (FlattenMaterial.DiffuseSamples.Num())
+			if (FlattenMaterial.DoesPropertyContainData(EFlattenMaterialProperties::Diffuse))
 			{
 				if (InMaterialLODSettings.ChannelsToCast[0].bBakeVertexColors)
 				{
 					SGMaterial->SetVertexColorChannel(SimplygonSDK::SG_MATERIAL_CHANNEL_BASECOLOR, 0);
 				}
 
-				SetMaterialChannelData(FlattenMaterial.DiffuseSamples, FlattenMaterial.RenderSize, SGMaterial, SimplygonSDK::SG_MATERIAL_CHANNEL_BASECOLOR);
+				SetMaterialChannelData(FlattenMaterial.GetPropertySamples(EFlattenMaterialProperties::Diffuse), FlattenMaterial.RenderSize, SGMaterial, SimplygonSDK::SG_MATERIAL_CHANNEL_BASECOLOR);
 			}
 
 			// Does current material have Metallic?
-			if (FlattenMaterial.MetallicSamples.Num())
+			if (FlattenMaterial.DoesPropertyContainData(EFlattenMaterialProperties::Metallic))
 			{
-				SetMaterialChannelData(FlattenMaterial.MetallicSamples, FlattenMaterial.RenderSize, SGMaterial, SimplygonSDK::SG_MATERIAL_CHANNEL_METALLIC);
+				SetMaterialChannelData(FlattenMaterial.GetPropertySamples(EFlattenMaterialProperties::Metallic), FlattenMaterial.RenderSize, SGMaterial, SimplygonSDK::SG_MATERIAL_CHANNEL_METALLIC);
 			}
 
 			// Does current material have Specular?
-			if (FlattenMaterial.SpecularSamples.Num())
+			if (FlattenMaterial.DoesPropertyContainData(EFlattenMaterialProperties::Specular))
 			{
-				SetMaterialChannelData(FlattenMaterial.SpecularSamples, FlattenMaterial.RenderSize, SGMaterial, SimplygonSDK::SG_MATERIAL_CHANNEL_SPECULAR);
+				SetMaterialChannelData(FlattenMaterial.GetPropertySamples(EFlattenMaterialProperties::Specular), FlattenMaterial.RenderSize, SGMaterial, SimplygonSDK::SG_MATERIAL_CHANNEL_SPECULAR);
 			}
 
 			// Does current material have Roughness?
-			if (FlattenMaterial.RoughnessSamples.Num())
+			if (FlattenMaterial.DoesPropertyContainData(EFlattenMaterialProperties::Roughness))
 			{
-				SetMaterialChannelData(FlattenMaterial.RoughnessSamples, FlattenMaterial.RenderSize, SGMaterial, SimplygonSDK::SG_MATERIAL_CHANNEL_ROUGHNESS);
+				SetMaterialChannelData(FlattenMaterial.GetPropertySamples(EFlattenMaterialProperties::Roughness), FlattenMaterial.RenderSize, SGMaterial, SimplygonSDK::SG_MATERIAL_CHANNEL_ROUGHNESS);
 			}
 
 			//Does current material have a normalmap?
-			if (FlattenMaterial.NormalSamples.Num())
+			if (FlattenMaterial.DoesPropertyContainData(EFlattenMaterialProperties::Normal))
 			{
-				SetMaterialChannelData(FlattenMaterial.NormalSamples, FlattenMaterial.RenderSize, SGMaterial, SimplygonSDK::SG_MATERIAL_CHANNEL_NORMALS);
+				SetMaterialChannelData(FlattenMaterial.GetPropertySamples(EFlattenMaterialProperties::Normal), FlattenMaterial.RenderSize, SGMaterial, SimplygonSDK::SG_MATERIAL_CHANNEL_NORMALS);
 			}
 
 			// Does current material have Opacity?
-			if (FlattenMaterial.OpacitySamples.Num())
+			if (FlattenMaterial.DoesPropertyContainData(EFlattenMaterialProperties::Opacity))
 			{
 #if USE_USER_OPACITY_CHANNEL
-				SetMaterialChannelData(FlattenMaterial.OpacitySamples, FlattenMaterial.RenderSize, SGMaterial, USER_MATERIAL_CHANNEL_OPACITY);
+				SetMaterialChannelData(FlattenMaterial.GetPropertySamples(EFlattenMaterialProperties::Opacity), FlattenMaterial.RenderSize, SGMaterial, USER_MATERIAL_CHANNEL_OPACITY);
 #else
-				SetMaterialChannelData(FlattenMaterial.OpacitySamples, FlattenMaterial.RenderSize, SGMaterial, SimplygonSDK::SG_MATERIAL_CHANNEL_OPACITY);
+				SetMaterialChannelData(FlattenMaterial.GetPropertySamples(EFlattenMaterialProperties::Opacity), FlattenMaterial.RenderSize, SGMaterial, SimplygonSDK::SG_MATERIAL_CHANNEL_OPACITY);
 #endif
 			}
 
-			if (FlattenMaterial.EmissiveSamples.Num())
+			if (FlattenMaterial.GetPropertySamples(EFlattenMaterialProperties::Emissive).Num())
 			{
-				SetMaterialChannelData(FlattenMaterial.EmissiveSamples, FlattenMaterial.RenderSize, SGMaterial, SimplygonSDK::SG_MATERIAL_CHANNEL_EMISSIVE);
+				SetMaterialChannelData(FlattenMaterial.GetPropertySamples(EFlattenMaterialProperties::Emissive), FlattenMaterial.RenderSize, SGMaterial, SimplygonSDK::SG_MATERIAL_CHANNEL_EMISSIVE);
 			}
 			else
 			{
@@ -2298,9 +2327,9 @@ private:
 				SetMaterialChannelData(BlackEmissive, FlattenMaterial.RenderSize, SGMaterial, SimplygonSDK::SG_MATERIAL_CHANNEL_EMISSIVE);
 			}
 
-			if (FlattenMaterial.SubSurfaceSamples.Num())
+			if (FlattenMaterial.DoesPropertyContainData(EFlattenMaterialProperties::SubSurface))
 			{
-				SetMaterialChannelData(FlattenMaterial.SubSurfaceSamples, FlattenMaterial.RenderSize, SGMaterial, USER_MATERIAL_CHANNEL_SUBSURFACE_COLOR);
+				SetMaterialChannelData(FlattenMaterial.GetPropertySamples(EFlattenMaterialProperties::Emissive), FlattenMaterial.RenderSize, SGMaterial, USER_MATERIAL_CHANNEL_SUBSURFACE_COLOR);
 			}
 
 			OutSGMaterialTable->AddMaterial(SGMaterial);
