Copying //UE4/Ocean-Staging to Dev-Main (Source: //Ocean/Main @ 3596817)

=====================================
 MAJOR FEATURES + CHANGES
=====================================

Change 3582844 by [REDACTED_PERSON]

	#ue4: Fixed null-pointer dereference crash in CookOnTheFlyServer.

Change 3580647 by josh.jensen

	Added support for incremental cooking via File->Cook Content

	The new setting can be found in the Project Settings, Editor->Cooker->Iterative cooking for the File->Cook Content menu item.

Change 3577413 by [REDACTED_PERSON]

	New project settings to globally disable/enable creation of depth-only and reversed index buffers. Disabling for Spyjinx to save memory.

Change 3575814 by [REDACTED_PERSON]

	Hotfix for incorrect resolution on iOS. Merged from 3534128.

	Merging

	//UE4/Release-4.17/Engine/Source/Programs/...

	to //Ocean/Main/Engine/Source/Programs/...

Change 3570141 by [REDACTED_PERSON]

	Hotfix from CL 3530807 for UE-47084: Check() when rendering wire frame

Change 3565495 by [REDACTED_PERSON]

	Hotfix from CL 3547311 for UE-47562
	Ocean was crashing as well whenever the -game runtime argument was used

Change 3565338 by [REDACTED_PERSON]

	Fixed a warning in StaticMeshResources.h on Mac & iOS builds.

Change 3564653 by [REDACTED_PERSON]

	Correct tracking of static mesh index memory.

Change 3557514 by [REDACTED_PERSON]

	Hotfix from CL 3551558 for UE-39404: UMG Designer: Named Slot doesn't pre-vis correctly with a Widget Switcher

Change 3549240 by [REDACTED_PERSON]

	#ue4: Added missing #include to AudioMixerSource.cpp.

	Fixes compiler error.

Change 3549207 by [REDACTED_PERSON]

	Fix for shutting down audio mixer while streaming a sound wave.

Change 3548753 by [REDACTED_PERSON]

	add some low resolution textures for shipping, point to their use for Ocean

Change 3547287 by [REDACTED_PERSON]

	Hotfix from CL 3543156 for UE-47321: Save on successful compile only not working in blueprints

Change 3545479 by [REDACTED_PERSON]

	Hotfix from CL 3545197 for UE-47431 Sound cues with delays do not play in the animation editor when set as anim notifies

Change 3545232 by [REDACTED_PERSON]

	Hotfix from CL 3525151 for UE-46871, which may help with blueprint compiling crashes

Change 3542258 by [REDACTED_PERSON]

	Hotfix from CL 3541404 for UE-47322: Breakpoints on Flow Control nodes are not triggered

Change 3537264 by [REDACTED_PERSON]

	Hotfix from CL 3536564 for UE-47225: UMG - Fixing a bug with the way text renamming occured.  We no longer show the metadata in the rename block when you go to edit, we show the text that you are actually editing.

Change 3535034 by [REDACTED_PERSON]

	Sequencer: Refix crash on end with null object in blending accumulator. Null objects are acceptable for master tracks. Added ensure for transform tracks, property, and skeletal mesh tracks on null object.

Change 3532229 by [REDACTED_PERSON]

	#ue4: Added cvar AudioThread.AboveNormalPriority that allows increasing the audio thread's priority.

	AudioThread.AboveNormalPriority = 0 // audio thread set to below normal priority (the default).
	AudioThread.AboveNormalPriority = 1 // audio thread set to above normal priority.

	(Josh's change, submitting for him.)

Change 3531562 by [REDACTED_PERSON]

	#ocn - Hotfix from CL 3521769 to fix OpenGL crashing on UnrealFrontend or Editor startup on Mac.

	Merging

	//UE4/Release-4.17/Engine/...

	to //Ocean/Main/Engine/...

Change 3531516 by [REDACTED_PERSON]

	#ocn - Hide the splash screen at startup.

Change 3527193 by [REDACTED_PERSON]

	Sequencer: Fix crash on end with null object in blending accumulator

Change 3519065 by [REDACTED_PERSON]

	#ue4 Quick fix for RuntimeAssetCache hanging the threads in the thread pool. (#jira OCN-7628).

	The application waits for all threads in the thread pool to abandon/finish their work. The hung threads prevent the application from shutting down.
	The hung threads were waiting on a condition that could never be met due to leaked memory. More fixes are needed to fix the leaks and make the code thread-safe.

Change 3510458 by [REDACTED_PERSON]

	#ocn - Don't return zero stride for compressed textures from FMetalSurface::Lock. Fixes broken iOS texture atlases.

Change 3510204 by [REDACTED_PERSON]

	#ocn Hotfix merged from cl 3510044. Fixes iOS slate scrollbox scrolling.

	Merging

	//UE4/Main/Engine/Source/Runtime/Slate/Private/Framework/Application/SlateApplication.cpp

	to //Ocean/Main/Engine/Source/Runtime/Slate/Private/Framework/Application/SlateApplication.cpp

Change 3483501 by [REDACTED_PERSON]

	Hotfix from CL 3483340 for Hot-Reload Asset Manager reinstancing issues (second half)

Change 3483500 by [REDACTED_PERSON]

	Hotfix from CL 3482498 for Hot-Reload Asset Manager reinstancing issues (first half)

diff --git a/Engine/Config/BaseEditor.ini b/Engine/Config/BaseEditor.ini
--- a/Engine/Config/BaseEditor.ini
+++ b/Engine/Config/BaseEditor.ini
@@ -394,6 +394,7 @@ bUseFullGCAssetClassNames=true
 +ConfigSettingBlacklist=*.Engine:OnlineSubsystem*:/Script/Engine.Engine
 +ConfigSettingBlacklist=*.Engine:Audio
 +ConfigSettingBlacklist=*.Engine:/Script/LinuxTargetPlatform.LinuxTargetSettings
++ConfigSettingBlacklist=*.Engine:/Script/UnrealEd.CookerSettings:bIterativeCookingForFileCookContent
 
 
 +VersionedIntRValues=r.AllowStaticLighting
@@ -529,4 +530,4 @@ bUseCompilationManager=true
 +DebugResolutions=(Width=1280,Height=720,Description="",Color=(R=0.000000,G=0.318547,B=0.806952,A=1.000000))
 +DebugResolutions=(Width=720,Height=1280,Description="",Color=(R=0.000000,G=0.318547,B=0.806952,A=1.000000))
 +DebugResolutions=(Width=1136,Height=640,Description="",Color=(R=0.000000,G=0.238398,B=0.658375,A=1.000000))
-+DebugResolutions=(Width=640,Height=1136,Description="",Color=(R=0.000000,G=0.238398,B=0.658375,A=1.000000))
\ No newline at end of file
++DebugResolutions=(Width=640,Height=1136,Description="",Color=(R=0.000000,G=0.238398,B=0.658375,A=1.000000))
diff --git a/Engine/Source/Developer/MeshUtilities/Private/MeshUtilities.cpp b/Engine/Source/Developer/MeshUtilities/Private/MeshUtilities.cpp
--- a/Engine/Source/Developer/MeshUtilities/Private/MeshUtilities.cpp
+++ b/Engine/Source/Developer/MeshUtilities/Private/MeshUtilities.cpp
@@ -125,6 +125,18 @@ static TAutoConsoleVariable<int32> CVarTriangleOrderOptimization(
 	TEXT("2: No triangle order optimization. (least efficient, debugging purposes only)"),
 	ECVF_Default);
 
+static TAutoConsoleVariable<int32> CVarSupportDepthOnlyIndexBuffers(
+	TEXT("r.SupportDepthOnlyIndexBuffers"),
+	1,
+	TEXT("Enables depth-only index buffers. Saves a little time at the expense of doubling the size of index buffers."),
+	ECVF_ReadOnly | ECVF_RenderThreadSafe);
+
+static TAutoConsoleVariable<int32> CVarSupportReversedIndexBuffers(
+	TEXT("r.SupportReversedIndexBuffers"),
+	1,
+	TEXT("Enables reversed index buffers. Saves a little time at the expense of doubling the size of index buffers."),
+	ECVF_ReadOnly | ECVF_RenderThreadSafe);
+
 IMPLEMENT_MODULE(FMeshUtilities, MeshUtilities);
 
 /*------------------------------------------------------------------------------
@@ -3106,9 +3118,9 @@ public:
 				}
 			}
 			LODModel.IndexBuffer.SetIndices(CombinedIndices, bNeeds32BitIndices ? EIndexBufferStride::Force32Bit : EIndexBufferStride::Force16Bit);
-
+			
 			// Build the reversed index buffer.
-			if (InOutModels[0].BuildSettings.bBuildReversedIndexBuffer)
+			if (InOutModels[0].BuildSettings.bBuildReversedIndexBuffer && MeshUtilities.bEnableReversedIndexBuffer)
 			{
 				TArray<uint32> InversedIndices;
 				const int32 IndexCount = CombinedIndices.Num();
@@ -3129,13 +3141,14 @@ public:
 
 			// Build the depth-only index buffer.
 			TArray<uint32> DepthOnlyIndices;
+			if (MeshUtilities.bEnableDepthOnlyIndexBuffer)
 			{
 				BuildDepthOnlyIndexBuffer(
 					DepthOnlyIndices,
 					Vertices,
 					CombinedIndices,
 					LODModel.Sections
-					);
+				);
 
 				if (DepthOnlyIndices.Num() < 50000 * 3)
 				{
@@ -3146,7 +3159,7 @@ public:
 			}
 
 			// Build the inversed depth only index buffer.
-			if (InOutModels[0].BuildSettings.bBuildReversedIndexBuffer)
+			if (InOutModels[0].BuildSettings.bBuildReversedIndexBuffer && MeshUtilities.bEnableDepthOnlyIndexBuffer && MeshUtilities.bEnableReversedIndexBuffer)
 			{
 				TArray<uint32> ReversedDepthOnlyIndices;
 				const int32 IndexCount = DepthOnlyIndices.Num();
@@ -5325,16 +5338,21 @@ void FMeshUtilities::StartupModule()
 
 	bUsingNvTriStrip = !bDisableTriangleOrderOptimization && (CVarTriangleOrderOptimization.GetValueOnGameThread() == 0);
 
+	bEnableDepthOnlyIndexBuffer = (CVarSupportDepthOnlyIndexBuffers.GetValueOnGameThread() == 1);
+	bEnableReversedIndexBuffer = (CVarSupportReversedIndexBuffers.GetValueOnGameThread() == 1);
+
 	IMeshReductionManagerModule& Module = FModuleManager::Get().LoadModuleChecked<IMeshReductionManagerModule>("MeshReductionInterface");
 	IMeshReduction* StaticMeshReduction = Module.GetStaticMeshReductionInterface();
 	
 
 	// Construct and cache the version string for the mesh utilities module.
 	VersionString = FString::Printf(
-		TEXT("%s%s%s"),
+		TEXT("%s%s%s%s%s"),
 		MESH_UTILITIES_VER,
 		StaticMeshReduction ? *StaticMeshReduction->GetVersionString() : TEXT(""),
-		bUsingNvTriStrip ? TEXT("_NvTriStrip") : TEXT("")
+		bUsingNvTriStrip ? TEXT("_NvTriStrip") : TEXT(""),
+		bEnableDepthOnlyIndexBuffer ? TEXT("_DepthOnlyIB") : TEXT("_NoDepthOnlyIB"),
+		bEnableReversedIndexBuffer ? TEXT("_ReversedIB") : TEXT("_NoReversedIB")
 		);
 
 	// hook up level editor extension for skeletal mesh conversion
diff --git a/Engine/Source/Developer/MeshUtilities/Private/MeshUtilitiesPrivate.h b/Engine/Source/Developer/MeshUtilities/Private/MeshUtilitiesPrivate.h
--- a/Engine/Source/Developer/MeshUtilities/Private/MeshUtilitiesPrivate.h
+++ b/Engine/Source/Developer/MeshUtilities/Private/MeshUtilitiesPrivate.h
@@ -64,6 +64,12 @@ private:
 	/** True if we disable triangle order optimization.  For debugging purposes only */
 	bool bDisableTriangleOrderOptimization;
 
+	/** True if depth-only index buffers are enabled */
+	bool bEnableDepthOnlyIndexBuffer;
+	
+	/** True if reversed index buffers are enabled */
+	bool bEnableReversedIndexBuffer;
+
 	// IMeshUtilities interface.
 	virtual const FString& GetVersionString() const override
 	{
diff --git a/Engine/Source/Editor/MainFrame/Private/Frame/MainFrameActions.cpp b/Engine/Source/Editor/MainFrame/Private/Frame/MainFrameActions.cpp
--- a/Engine/Source/Editor/MainFrame/Private/Frame/MainFrameActions.cpp
+++ b/Engine/Source/Editor/MainFrame/Private/Frame/MainFrameActions.cpp
@@ -31,6 +31,7 @@
 #include "EditorStyleSet.h"
 #include "Editor/EditorPerProjectUserSettings.h"
 #include "Settings/EditorExperimentalSettings.h"
+#include "CookerSettings.h"
 #include "UnrealEdMisc.h"
 #include "FileHelpers.h"
 #include "Dialogs/Dialogs.h"
@@ -392,6 +393,12 @@ void FMainFrameActionCallbacks::CookContent(const FName InPlatformInfoName)
 
 	OptionalParams += GetCookingOptionalParams();
 
+	UCookerSettings const* CookerSettings = GetDefault<UCookerSettings>();
+	if (CookerSettings->bIterativeCookingForFileCookContent)
+	{
+		OptionalParams += TEXT(" -iterate");
+	}
+
 	if (FApp::IsRunningDebug())
 	{
 		OptionalParams += TEXT(" -UseDebugParamForEditorExe");
diff --git a/Engine/Source/Editor/UnrealEd/Private/CookOnTheFlyServer.cpp b/Engine/Source/Editor/UnrealEd/Private/CookOnTheFlyServer.cpp
--- a/Engine/Source/Editor/UnrealEd/Private/CookOnTheFlyServer.cpp
+++ b/Engine/Source/Editor/UnrealEd/Private/CookOnTheFlyServer.cpp
@@ -4444,15 +4444,15 @@ void UCookOnTheFlyServer::PopulateCookedPackagesFromDisk(const TArray<ITargetPla
 			const FName CookedFile = CookedPaths.Value;
 			const FName UncookedFilename = CookedPaths.Key;
 			const FName* FoundPackageName = GetCachedPackageFilenameToPackageFName(UncookedFilename);
+			const FName PackageName = FoundPackageName ? *FoundPackageName : FName();
+			bool bShouldKeep = true;
+
 			if ( !FoundPackageName )
 			{
 				++NumPackagesRemoved;
+				bShouldKeep = false;
 			}
-			const FName PackageName = *FoundPackageName;
-			const FString UncookedFilenameString = UncookedFilename.ToString();
-			bool bShouldKeep = true;
-
-			if (ModifiedPackages.Contains(PackageName))
+			else if (ModifiedPackages.Contains(PackageName))
 			{
 				++NumPackagesFileHashMismatch;
 				bShouldKeep = false;
@@ -4511,6 +4511,7 @@ void UCookOnTheFlyServer::PopulateCookedPackagesFromDisk(const TArray<ITargetPla
 				{
 					// the cooker should rebuild this package because it's not in the cooked package list
 					// the new package will have higher priority then the package in the original shared cooked build
+					const FString UncookedFilenameString = UncookedFilename.ToString();
 					UE_LOG(LogCook, Verbose, TEXT("Shared cooked build: Detected package is out of date %s"), *UncookedFilenameString);
 				}
 			}
diff --git a/Engine/Source/Editor/UnrealEd/Public/CookerSettings.h b/Engine/Source/Editor/UnrealEd/Public/CookerSettings.h
--- a/Engine/Source/Editor/UnrealEd/Public/CookerSettings.h
+++ b/Engine/Source/Editor/UnrealEd/Public/CookerSettings.h
@@ -36,6 +36,10 @@ public:
 	UPROPERTY(GlobalConfig, EditAnywhere, Category = Cooker, meta = (DisplayName = "Iterative cooking for builds launched from the editor (launch on)"))
 	bool bIterativeCookingForLaunchOn;
 
+	/** Enable -iterate for launch on */
+	UPROPERTY(GlobalConfig, EditAnywhere, Category = Cooker, meta = (DisplayName = "Iterative cooking for the File->Cook Content menu item"))
+	bool bIterativeCookingForFileCookContent;
+
 	/** Ignore ini changes when doing iterative cooking, either in editor or out of editor */
 	UPROPERTY(GlobalConfig, EditAnywhere, Category = Cooker, AdvancedDisplay)
 	bool bIgnoreIniSettingsOutOfDateForIteration;
diff --git a/Engine/Source/Runtime/Apple/MetalRHI/Private/MetalTexture.cpp b/Engine/Source/Runtime/Apple/MetalRHI/Private/MetalTexture.cpp
--- a/Engine/Source/Runtime/Apple/MetalRHI/Private/MetalTexture.cpp
+++ b/Engine/Source/Runtime/Apple/MetalRHI/Private/MetalTexture.cpp
@@ -1367,12 +1367,16 @@ void* FMetalSurface::Lock(uint32 MipIndex, uint32 ArrayIndex, EResourceLockMode
 				//kick the current command buffer.
 				GetMetalDeviceContext().SubmitCommandBufferAndWait();
 #endif
+				
+				// This block breaks the texture atlas system in Ocean, which depends on nonzero strides coming back from compressed textures. Turning off.
+#if 0
 				if (PixelFormat == PF_PVRTC2 || PixelFormat == PF_PVRTC4)
 				{
 					// for compressed textures metal debug RT expects 0 for rowBytes and imageBytes.
 					DestStride = 0;
 					MipBytes = 0;
 				}
+#endif
 				uint32 BytesPerRow = DestStride;
 				if (PixelFormat == PF_PVRTC2 || PixelFormat == PF_PVRTC4)
 				{
diff --git a/Engine/Source/Runtime/Engine/Classes/Engine/RendererSettings.h b/Engine/Source/Runtime/Engine/Classes/Engine/RendererSettings.h
--- a/Engine/Source/Runtime/Engine/Classes/Engine/RendererSettings.h
+++ b/Engine/Source/Runtime/Engine/Classes/Engine/RendererSettings.h
@@ -510,6 +510,17 @@ class ENGINE_API URendererSettings : public UDeveloperSettings
 		ConfigRestartRequired = true))
 		uint32 bGPUSkinLimit2BoneInfluences : 1;
 
+	UPROPERTY(config, EditAnywhere, Category = Optimizations, meta = (
+		ConsoleVariable = "r.SupportDepthOnlyIndexBuffers", DisplayName = "Support depth-only index buffers",
+		ToolTip = "Support depth-only index buffers, which provide a minor rendering speedup at the expense of using twice the index buffer memory.",
+		ConfigRestartRequired = true))
+		uint32 bSupportDepthOnlyIndexBuffers : 1;
+
+	UPROPERTY(config, EditAnywhere, Category = Optimizations, meta = (
+		ConsoleVariable = "r.SupportReversedIndexBuffers", DisplayName = "Support reversed index buffers",
+		ToolTip = "Support reversed index buffers, which provide a minor rendering speedup at the expense of using twice the index buffer memory.",
+		ConfigRestartRequired = true))
+		uint32 bSupportReversedIndexBuffers : 1;
 public:
 
 	//~ Begin UObject Interface
diff --git a/Engine/Source/Runtime/Engine/Private/StaticMesh.cpp b/Engine/Source/Runtime/Engine/Private/StaticMesh.cpp
--- a/Engine/Source/Runtime/Engine/Private/StaticMesh.cpp
+++ b/Engine/Source/Runtime/Engine/Private/StaticMesh.cpp
@@ -419,6 +419,9 @@ FStaticMeshLODResources::FStaticMeshLODResources()
 	, DepthOnlyNumTriangles(0)
 	, SplineVertexFactory(nullptr)
 	, SplineVertexFactoryOverrideColorVertexBuffer(nullptr)
+#if STATS
+	, StaticMeshIndexMemory(0)
+#endif
 {
 }
 
@@ -446,6 +449,18 @@ void FStaticMeshLODResources::InitResources(UStaticMesh* Parent)
 			UE_LOG(LogStaticMesh, Warning, TEXT("[%s] Mesh has more that 65535 vertices, incompatible with mobile; forcing 16-bit (will probably cause rendering issues)." ), *Parent->GetName());
 		}
 	}
+
+#if STATS
+	uint32 iMem = IndexBuffer.GetAllocatedSize();
+	uint32 wiMem = WireframeIndexBuffer.GetAllocatedSize();
+	uint32 riMem = ReversedIndexBuffer.GetAllocatedSize();
+	uint32 doiMem = DepthOnlyIndexBuffer.GetAllocatedSize();
+	uint32 rdoiMem = ReversedDepthOnlyIndexBuffer.GetAllocatedSize();
+	uint32 aiMem = AdjacencyIndexBuffer.GetAllocatedSize();
+	StaticMeshIndexMemory = iMem + wiMem + riMem + doiMem + rdoiMem + aiMem;
+	INC_DWORD_STAT_BY(STAT_StaticMeshIndexMemory, StaticMeshIndexMemory);
+#endif
+
 	BeginInitResource(&IndexBuffer);
 	if( WireframeIndexBuffer.GetNumIndices() > 0 )
 	{
@@ -497,14 +512,10 @@ void FStaticMeshLODResources::InitResources(UStaticMesh* Parent)
 			const uint32 StaticMeshVertexMemory =
 			This->VertexBuffer.GetStride() * This->VertexBuffer.GetNumVertices() +
 			This->PositionVertexBuffer.GetStride() * This->PositionVertexBuffer.GetNumVertices();
-			const uint32 StaticMeshIndexMemory = This->IndexBuffer.GetAllocatedSize()
-				+ This->WireframeIndexBuffer.GetAllocatedSize()
-				+ (RHISupportsTessellation( GShaderPlatformForFeatureLevel[GMaxRHIFeatureLevel] ) ? This->AdjacencyIndexBuffer.GetAllocatedSize() : 0);
 			const uint32 ResourceVertexColorMemory = This->ColorVertexBuffer.GetStride() * This->ColorVertexBuffer.GetNumVertices();
 
 			INC_DWORD_STAT_BY( STAT_StaticMeshVertexMemory, StaticMeshVertexMemory );
 			INC_DWORD_STAT_BY( STAT_ResourceVertexColorMemory, ResourceVertexColorMemory );
-			INC_DWORD_STAT_BY( STAT_StaticMeshIndexMemory, StaticMeshIndexMemory );
 		});
 }
 
@@ -513,9 +524,6 @@ void FStaticMeshLODResources::ReleaseResources()
 	const uint32 StaticMeshVertexMemory = 
 		VertexBuffer.GetStride() * VertexBuffer.GetNumVertices() + 
 		PositionVertexBuffer.GetStride() * PositionVertexBuffer.GetNumVertices();
-	const uint32 StaticMeshIndexMemory = IndexBuffer.GetAllocatedSize()
-		+ WireframeIndexBuffer.GetAllocatedSize()
-		+ (RHISupportsTessellation( GShaderPlatformForFeatureLevel[GMaxRHIFeatureLevel] ) ? AdjacencyIndexBuffer.GetAllocatedSize() : 0);
 	const uint32 ResourceVertexColorMemory = ColorVertexBuffer.GetStride() * ColorVertexBuffer.GetNumVertices();
 
 	DEC_DWORD_STAT_BY( STAT_StaticMeshVertexMemory, StaticMeshVertexMemory );
diff --git a/Engine/Source/Runtime/Engine/Public/StaticMeshResources.h b/Engine/Source/Runtime/Engine/Public/StaticMeshResources.h
--- a/Engine/Source/Runtime/Engine/Public/StaticMeshResources.h
+++ b/Engine/Source/Runtime/Engine/Public/StaticMeshResources.h
@@ -270,6 +270,10 @@ struct FStaticMeshLODResources
 	struct FSplineMeshVertexFactory* SplineVertexFactory;
 	struct FSplineMeshVertexFactory* SplineVertexFactoryOverrideColorVertexBuffer;
 
+#if STATS
+	uint32 StaticMeshIndexMemory;
+#endif
+	
 	/** Default constructor. */
 	FStaticMeshLODResources();
 
diff --git a/Engine/Source/Runtime/RuntimeAssetCache/Private/RuntimeAssetCacheAsyncWorker.cpp b/Engine/Source/Runtime/RuntimeAssetCache/Private/RuntimeAssetCacheAsyncWorker.cpp
--- a/Engine/Source/Runtime/RuntimeAssetCache/Private/RuntimeAssetCacheAsyncWorker.cpp
+++ b/Engine/Source/Runtime/RuntimeAssetCache/Private/RuntimeAssetCacheAsyncWorker.cpp
@@ -62,9 +62,11 @@ void FRuntimeAssetCacheAsyncWorker::DoWork()
 	}
 	else
 	{
-		while (Metadata->IsBuilding())
+		if (Metadata->IsBuilding())
 		{
-			FPlatformProcess::SleepNoStats(0.0f);
+			// Another worker is building this asset.
+			Data = nullptr;
+			return;
 		}
 
 		Metadata = FRuntimeAssetCacheBackend::Get().GetCachedData(BucketName, *CacheKey, Data, DataSize);
