Make sure meshutilities use the correct high level FRawMesh API

diff --git a/Engine/Source/Developer/MeshUtilities/Private/MeshUtilities.cpp b/Engine/Source/Developer/MeshUtilities/Private/MeshUtilities.cpp
--- a/Engine/Source/Developer/MeshUtilities/Private/MeshUtilities.cpp
+++ b/Engine/Source/Developer/MeshUtilities/Private/MeshUtilities.cpp
@@ -948,7 +948,7 @@ UStaticMesh* FMeshUtilities::ConvertMeshesToStaticMesh(const TArray<UMeshCompone
 					SrcModel->BuildSettings.bGenerateLightmapUVs = true;
 					SrcModel->BuildSettings.SrcLightmapIndex = 0;
 					SrcModel->BuildSettings.DstLightmapIndex = LightMapIndex;
-					SrcModel->RawMeshBulkData->SaveRawMesh(RawMesh);
+					SrcModel->SaveRawMesh(RawMesh);
 				}
 			}
 
@@ -2856,9 +2856,9 @@ public:
 			FRawMesh& RawMesh = *new(LODMeshes)FRawMesh;
 			TMultiMap<int32, int32>& OverlappingCorners = *new(LODOverlappingCorners)TMultiMap<int32, int32>;
 
-			if (!SrcModel.RawMeshBulkData->IsEmpty())
+			if (!SrcModel.IsRawMeshEmpty())
 			{
-				SrcModel.RawMeshBulkData->LoadRawMesh(RawMesh);
+				SrcModel.LoadRawMesh(RawMesh);
 				// Make sure the raw mesh is not irreparably malformed.
 				if (!RawMesh.IsValidOrFixable())
 				{
@@ -3255,7 +3255,7 @@ public:
 		{
 			if (!HasRawMesh[Index])
 			{
-				SourceModels[Index].RawMeshBulkData->SaveRawMesh(LODMeshes[Index]);
+				SourceModels[Index].SaveRawMesh(LODMeshes[Index]);
 				bDirty = true;
 			}
 		}
