Issue a warning to the user when skeletal mesh legacy "built in" build end up with zero or nan tangents.

diff --git a/Engine/Source/Developer/MeshUtilities/Private/MeshUtilities.cpp b/Engine/Source/Developer/MeshUtilities/Private/MeshUtilities.cpp
--- a/Engine/Source/Developer/MeshUtilities/Private/MeshUtilities.cpp
+++ b/Engine/Source/Developer/MeshUtilities/Private/MeshUtilities.cpp
@@ -4662,6 +4662,7 @@ bool FMeshUtilities::BuildSkeletalMesh_Legacy(FSkeletalMeshLODModel& LODModel
 	TArray<FSoftSkinBuildVertex> RawVertices;
 	RawVertices.Reserve(Points.Num());
 
+	int32 NTBErrorCount = 0;
 	// Create a list of vertex Z/index pairs
 
 	for (int32 FaceIndex = 0; FaceIndex < Faces.Num(); FaceIndex++)
@@ -4798,6 +4799,19 @@ bool FMeshUtilities::BuildSkeletalMesh_Legacy(FSkeletalMeshLODModel& LODModel
 				TangentZ.Normalize();
 			}
 
+			if (TangentX.IsNearlyZero() || TangentX.ContainsNaN())
+			{
+				NTBErrorCount++;
+			}
+			if (TangentY.IsNearlyZero() || TangentY.ContainsNaN())
+			{
+				NTBErrorCount++;
+			}
+			if (TangentZ.IsNearlyZero() || TangentZ.ContainsNaN())
+			{
+				NTBErrorCount++;
+			}
+
 			Vertex.TangentX = TangentX;
 			Vertex.TangentY = TangentY;
 			Vertex.TangentZ = TangentZ;
@@ -4857,6 +4871,15 @@ bool FMeshUtilities::BuildSkeletalMesh_Legacy(FSkeletalMeshLODModel& LODModel
 		}
 	}
 
+	if (NTBErrorCount > 0 && OutWarningMessages)
+	{
+		OutWarningMessages->Add(FText::FromString(TEXT("SkeletalMesh compute tangents [built in]: Build result data contain 0 or NAN tangent value. Bad tangent value will impact shading.")));
+		if (OutWarningNames)
+		{
+			OutWarningNames->Add(FFbxErrors::Generic_Mesh_TangentsComputeError);
+		}
+	}
+
 	// Generate chunks and their vertices and indices
 	SkeletalMeshTools::BuildSkeletalMeshChunks(Faces, RawVertices, VertIndexAndZ, OverlappingThresholds, Chunks, bTooManyVerts);
 
diff --git a/Engine/Source/Runtime/Core/Private/Misc/FbxErrors.cpp b/Engine/Source/Runtime/Core/Private/Misc/FbxErrors.cpp
--- a/Engine/Source/Runtime/Core/Private/Misc/FbxErrors.cpp
+++ b/Engine/Source/Runtime/Core/Private/Misc/FbxErrors.cpp
@@ -51,6 +51,8 @@ FName FFbxErrors::Generic_Mesh_SkinxxNameError(TEXT("Generic_Mesh_SkinxxNameErro
 
 FName FFbxErrors::Generic_Mesh_TooManyLODs(TEXT("Generic_Mesh_TooManyLODs"));
 
+FName FFbxErrors::Generic_Mesh_TangentsComputeError(TEXT("Generic_Mesh_TangentsComputeError"));
+
 /** Static Mesh **/
 FName FFbxErrors::StaticMesh_TooManyMaterials(TEXT("StaticMesh_TooManyMaterials"));
 
diff --git a/Engine/Source/Runtime/Core/Public/Misc/FbxErrors.h b/Engine/Source/Runtime/Core/Public/Misc/FbxErrors.h
--- a/Engine/Source/Runtime/Core/Public/Misc/FbxErrors.h
+++ b/Engine/Source/Runtime/Core/Public/Misc/FbxErrors.h
@@ -63,6 +63,8 @@ struct CORE_API FFbxErrors
 
 	static FName Generic_Mesh_TooManyLODs;
 
+	static FName Generic_Mesh_TangentsComputeError;
+
 	/** Static Mesh **/
 	static FName StaticMesh_TooManyMaterials;
 
