Added mesh distance field downsampling capability during serialization time (compressed and not compressed)
Added DownSampleMeshDistanceFieldDivider target platform variable

#rnx

diff --git a/Engine/Source/Developer/MeshUtilities/Private/MeshDistanceFieldUtilities.cpp b/Engine/Source/Developer/MeshUtilities/Private/MeshDistanceFieldUtilities.cpp
--- a/Engine/Source/Developer/MeshUtilities/Private/MeshDistanceFieldUtilities.cpp
+++ b/Engine/Source/Developer/MeshUtilities/Private/MeshDistanceFieldUtilities.cpp
@@ -748,3 +748,138 @@ void FMeshUtilities::GenerateSignedDistanceFieldVolumeData(
 }
 
 #endif // PLATFORM_ENABLE_VECTORINTRINSICS
+
+void FMeshUtilities::DownSampleDistanceFieldVolumeData(FDistanceFieldVolumeData& DistanceFieldData, float Divider)
+{
+	static const auto CVarCompress = IConsoleManager::Get().FindTConsoleVariableDataInt(TEXT("r.DistanceFieldBuild.Compress"));
+	const bool bDataIsCompressed = CVarCompress->GetValueOnAnyThread() != 0;
+
+	static const auto CVarEightBit = IConsoleManager::Get().FindTConsoleVariableDataInt(TEXT("r.DistanceFieldBuild.EightBit"));
+	const bool bDataIsEightBit = CVarEightBit->GetValueOnAnyThread() != 0;
+
+	if (!bDataIsEightBit)
+	{
+		return;
+	}
+
+	TArray<uint8> UncompressedData;
+	const TArray<uint8>* SourceData = &DistanceFieldData.CompressedDistanceFieldVolume;
+
+	if (DistanceFieldData.Size.X <= 4 || DistanceFieldData.Size.Y <= 4 || DistanceFieldData.Size.Z <= 4)
+	{
+		return;
+	}
+
+	int32 UncompressedSize = DistanceFieldData.Size.X * DistanceFieldData.Size.Y * DistanceFieldData.Size.Z;
+
+	if (bDataIsCompressed)
+	{
+		UncompressedData.AddUninitialized(UncompressedSize);
+		verify(FCompression::UncompressMemory(NAME_Zlib, UncompressedData.GetData(), UncompressedSize, DistanceFieldData.CompressedDistanceFieldVolume.GetData(), DistanceFieldData.CompressedDistanceFieldVolume.Num()));
+		SourceData = &UncompressedData;
+	}
+
+	const FIntVector SrcSize = DistanceFieldData.Size;
+
+	FVector DstSizeFloat = FVector(SrcSize);
+	DstSizeFloat /= Divider;
+	
+	FIntVector DstSize;
+	DstSize.X = FMath::TruncToInt(DstSizeFloat.X);
+	DstSize.Y = FMath::TruncToInt(DstSizeFloat.Y);
+	DstSize.Z = FMath::TruncToInt(DstSizeFloat.Z);
+
+	FIntVector IntVectorOne = FIntVector(1);
+	FIntVector SrcSizeMinusOne = SrcSize - IntVectorOne;
+	FIntVector DstSizeMinusOne = DstSize - IntVectorOne;
+	
+	FVector Divider3 = FVector(SrcSizeMinusOne) / FVector(DstSizeMinusOne);
+
+	TArray<uint8> DownSampledTexture;
+	DownSampledTexture.SetNum(DstSize.X * DstSize.Y * DstSize.Z);
+
+	int32 SrcPitchX = SrcSize.X;
+	int32 SrcPitchY = SrcSize.X * SrcSize.Y;
+	int32 DstPitchX = DstSize.X;
+	int32 DstPitchY = DstSize.X * DstSize.Y;
+
+	for (int32 DstPosZ = 0; DstPosZ < DstSize.Z; ++DstPosZ)
+	{
+		int32 SrcPosZ0 = FMath::TruncToInt((float(DstPosZ) + 0.0f) * Divider3.Z);
+		int32 SrcPosZ1 = FMath::TruncToInt((float(DstPosZ) + 1.0f) * Divider3.Z);
+		SrcPosZ1 = FMath::Min(SrcPosZ1, SrcSizeMinusOne.Z);
+
+		int32 SrcOffsetZ0 = SrcPosZ0 * SrcPitchY;
+		int32 SrcOffsetZ1 = SrcPosZ1 * SrcPitchY;
+		
+		int32 DstOffsetZ = DstPosZ * DstPitchY;
+
+		for (int32 DstPosY = 0; DstPosY < DstSize.Y; ++DstPosY)
+		{
+			int32 SrcPosY0 = FMath::TruncToInt(float(DstPosY) + 0.0f) * Divider3.Y;
+			int32 SrcPosY1 = FMath::TruncToInt(float(DstPosY) + 1.0f) * Divider3.Y;
+			SrcPosY1 = FMath::Min(SrcPosY1, SrcSizeMinusOne.Y);
+
+			int32 SrcOffsetZ0Y0 = SrcPosY0 * SrcPitchX + SrcOffsetZ0;
+			int32 SrcOffsetZ0Y1 = SrcPosY1 * SrcPitchX + SrcOffsetZ0;
+			int32 SrcOffsetZ1Y0 = SrcPosY0 * SrcPitchX + SrcOffsetZ1;
+			int32 SrcOffsetZ1Y1 = SrcPosY1 * SrcPitchX + SrcOffsetZ1;
+
+			int32 DstOffsetZY = DstPosY * DstPitchX + DstOffsetZ;
+
+			for (int32 DstPosX = 0; DstPosX < DstSize.X; ++DstPosX)
+			{
+				int32 SrcPosX0 = FMath::TruncToInt(float(DstPosX) + 0.0f) * Divider3.X;
+				int32 SrcPosX1 = FMath::TruncToInt(float(DstPosX) + 1.0f) * Divider3.X;
+				SrcPosX1 = FMath::Min(SrcPosX1, SrcSizeMinusOne.X - 1);
+
+				uint32 a = uint32((*SourceData)[SrcPosX0 + SrcOffsetZ0Y0]);
+				uint32 b = uint32((*SourceData)[SrcPosX1 + SrcOffsetZ0Y0]);
+				uint32 c = uint32((*SourceData)[SrcPosX0 + SrcOffsetZ0Y1]);
+				uint32 d = uint32((*SourceData)[SrcPosX1 + SrcOffsetZ0Y1]);
+
+				uint32 e = uint32((*SourceData)[SrcPosX0 + SrcOffsetZ1Y0]);
+				uint32 f = uint32((*SourceData)[SrcPosX1 + SrcOffsetZ1Y0]);
+				uint32 g = uint32((*SourceData)[SrcPosX0 + SrcOffsetZ1Y1]);
+				uint32 h = uint32((*SourceData)[SrcPosX1 + SrcOffsetZ1Y1]);
+
+
+				a = a + b + c + d + e + f + g + h;
+				a /= 8;
+
+				DownSampledTexture[DstOffsetZY + DstPosX] = uint8(a);
+			}
+		}
+	}
+
+	DistanceFieldData.Size = DstSize;
+	UncompressedSize = DownSampledTexture.Num();
+
+	if (bDataIsCompressed)
+	{
+		TArray<uint8> TempCompressedMemory;
+		// Compressed can be slightly larger than uncompressed
+		TempCompressedMemory.Empty(UncompressedSize * 4 / 3);
+		TempCompressedMemory.AddUninitialized(UncompressedSize * 4 / 3);
+		int32 CompressedSize = TempCompressedMemory.Num() * TempCompressedMemory.GetTypeSize();
+
+		verify(FCompression::CompressMemory(
+			NAME_Zlib,
+			TempCompressedMemory.GetData(),
+			CompressedSize,
+			DownSampledTexture.GetData(),
+			UncompressedSize,
+			COMPRESS_BiasMemory));
+
+		DistanceFieldData.CompressedDistanceFieldVolume.Empty(CompressedSize);
+		DistanceFieldData.CompressedDistanceFieldVolume.AddUninitialized(CompressedSize);
+
+		FPlatformMemory::Memcpy(DistanceFieldData.CompressedDistanceFieldVolume.GetData(), TempCompressedMemory.GetData(), CompressedSize);
+	}
+	else
+	{
+		DistanceFieldData.CompressedDistanceFieldVolume.Empty(UncompressedSize);
+		DistanceFieldData.CompressedDistanceFieldVolume.AddUninitialized(UncompressedSize);
+		FPlatformMemory::Memcpy(DistanceFieldData.CompressedDistanceFieldVolume.GetData(), DownSampledTexture.GetData(), UncompressedSize);
+	}
+}
diff --git a/Engine/Source/Developer/MeshUtilities/Private/MeshUtilitiesPrivate.h b/Engine/Source/Developer/MeshUtilities/Private/MeshUtilitiesPrivate.h
--- a/Engine/Source/Developer/MeshUtilities/Private/MeshUtilitiesPrivate.h
+++ b/Engine/Source/Developer/MeshUtilities/Private/MeshUtilitiesPrivate.h
@@ -99,6 +99,8 @@ private:
 		float DistanceFieldResolutionScale,
 		bool bGenerateAsIfTwoSided,
 		FDistanceFieldVolumeData& OutData) override;
+	
+	virtual void DownSampleDistanceFieldVolumeData(class FDistanceFieldVolumeData& DistanceFieldData, float Divider) override;
 
 	virtual void RecomputeTangentsAndNormalsForRawMesh(bool bRecomputeTangents, bool bRecomputeNormals, const FMeshBuildSettings& InBuildSettings, FRawMesh &OutRawMesh) const override;
 	virtual void RecomputeTangentsAndNormalsForRawMesh(bool bRecomputeTangents, bool bRecomputeNormals, const FMeshBuildSettings& InBuildSettings, const FOverlappingCorners& InOverlappingCorners, FRawMesh &OutRawMesh) const override;
diff --git a/Engine/Source/Developer/MeshUtilities/Public/MeshUtilities.h b/Engine/Source/Developer/MeshUtilities/Public/MeshUtilities.h
--- a/Engine/Source/Developer/MeshUtilities/Public/MeshUtilities.h
+++ b/Engine/Source/Developer/MeshUtilities/Public/MeshUtilities.h
@@ -200,6 +200,16 @@ public:
 		bool bGenerateAsIfTwoSided,
 		class FDistanceFieldVolumeData& OutData) = 0;
 
+	/** 
+	 * Down sample distance field volume. 
+	 * Method overwrites data of DistanceFieldData. 
+	 * If input is compressed, it will be decompressed, downsampled and recompressed
+	 */
+	virtual void DownSampleDistanceFieldVolumeData(
+		class FDistanceFieldVolumeData& DistanceFieldData,
+		float Divider) = 0;
+
+
 	/** Helper structure for skeletal mesh import options */
 	struct MeshBuildOptions
 	{
diff --git a/Engine/Source/Developer/TargetPlatform/Private/TargetPlatformBase.cpp b/Engine/Source/Developer/TargetPlatform/Private/TargetPlatformBase.cpp
--- a/Engine/Source/Developer/TargetPlatform/Private/TargetPlatformBase.cpp
+++ b/Engine/Source/Developer/TargetPlatform/Private/TargetPlatformBase.cpp
@@ -33,6 +33,11 @@ bool FTargetPlatformBase::UsesDistanceFields() const
 	return true;
 }
 
+float FTargetPlatformBase::GetDownSampleMeshDistanceFieldDivider() const
+{
+	return 1.0f;
+}
+
 TSharedPtr<IDeviceManagerCustomPlatformWidgetCreator> FTargetPlatformBase::GetCustomWidgetCreator() const
 {
 	static TSharedPtr<FDeviceBrowserDefaultPlatformWidgetCreator> DefaultWidgetCreator = MakeShared<FDeviceBrowserDefaultPlatformWidgetCreator>();
diff --git a/Engine/Source/Developer/TargetPlatform/Public/Common/TargetPlatformBase.h b/Engine/Source/Developer/TargetPlatform/Public/Common/TargetPlatformBase.h
--- a/Engine/Source/Developer/TargetPlatform/Public/Common/TargetPlatformBase.h
+++ b/Engine/Source/Developer/TargetPlatform/Public/Common/TargetPlatformBase.h
@@ -41,6 +41,8 @@ public:
 	
 	TARGETPLATFORM_API virtual bool UsesDistanceFields() const override;
 
+	TARGETPLATFORM_API virtual float GetDownSampleMeshDistanceFieldDivider() const override;
+
 #if WITH_ENGINE
 	virtual void GetReflectionCaptureFormats( TArray<FName>& OutFormats ) const override
 	{
diff --git a/Engine/Source/Developer/TargetPlatform/Public/Interfaces/ITargetPlatform.h b/Engine/Source/Developer/TargetPlatform/Public/Interfaces/ITargetPlatform.h
--- a/Engine/Source/Developer/TargetPlatform/Public/Interfaces/ITargetPlatform.h
+++ b/Engine/Source/Developer/TargetPlatform/Public/Interfaces/ITargetPlatform.h
@@ -358,6 +358,13 @@ public:
 	*/
 	virtual bool UsesDistanceFields() const = 0;
 
+	/**
+	* Gets down sample mesh distance field divider.
+	*
+	* @return 1 if platform does not need to downsample mesh distance fields
+	*/
+	virtual float GetDownSampleMeshDistanceFieldDivider() const = 0;
+
 #if WITH_ENGINE
 	/**
 	 * Gets the format to use for a particular body setup.
diff --git a/Engine/Source/Runtime/Engine/Private/StaticMesh.cpp b/Engine/Source/Runtime/Engine/Private/StaticMesh.cpp
--- a/Engine/Source/Runtime/Engine/Private/StaticMesh.cpp
+++ b/Engine/Source/Runtime/Engine/Private/StaticMesh.cpp
@@ -1288,18 +1288,41 @@ void FStaticMeshRenderData::Serialize(FArchive& Ar, UStaticMesh* Owner, bool bCo
 			{
 				FStaticMeshLODResources& LOD = LODResources[ResourceIndex];
 				
-				bool bValid = (LOD.DistanceFieldData != NULL);
+				bool bValid = (LOD.DistanceFieldData != nullptr);
 
 				Ar << bValid;
 
 				if (bValid)
 				{
-					if (!LOD.DistanceFieldData)
+#if WITH_EDITOR
+					bool bDownSampling = Ar.IsCooking() && Ar.IsSaving();
+					
+					if (bDownSampling)
 					{
-						LOD.DistanceFieldData = new FDistanceFieldVolumeData();
+						float Divider = Ar.CookingTarget()->GetDownSampleMeshDistanceFieldDivider();
+						bDownSampling = Divider > 1;
+
+						if (bDownSampling)
+						{
+							FDistanceFieldVolumeData DownSampledDFVolumeData = *LOD.DistanceFieldData;
+							IMeshUtilities& MeshUtilities = FModuleManager::Get().LoadModuleChecked<IMeshUtilities>(TEXT("MeshUtilities"));
+
+							MeshUtilities.DownSampleDistanceFieldVolumeData(DownSampledDFVolumeData, Divider);
+
+							Ar << DownSampledDFVolumeData;
+						}
 					}
 
-					Ar << *(LOD.DistanceFieldData);
+					if (!bDownSampling)
+#endif
+					{
+						if (!LOD.DistanceFieldData)
+						{
+							LOD.DistanceFieldData = new FDistanceFieldVolumeData();
+						}
+
+						Ar << *(LOD.DistanceFieldData);
+					}
 				}
 			}
 		}
