Fix skeletalmesh compute normal with mikkt to use the connected triangle island with smoothing group instead of only smoothing group.

diff --git a/Engine/Source/Developer/MeshUtilities/Private/MeshUtilities.cpp b/Engine/Source/Developer/MeshUtilities/Private/MeshUtilities.cpp
--- a/Engine/Source/Developer/MeshUtilities/Private/MeshUtilities.cpp
+++ b/Engine/Source/Developer/MeshUtilities/Private/MeshUtilities.cpp
@@ -3302,6 +3302,93 @@ public:
 		}
 	}
 
+	//This function add every triangles connected to the triangle queue.
+	//A connected triangle pair must share at least 1 vertex between the two triangles.
+	//If bConnectByEdge is true, the connected triangle must share at least one edge (two vertex index)
+	void AddAdjacentFace(IMeshBuildData* BuildData, TBitArray<>& FaceAdded, TMap<int32, TArray<int32>>& VertexIndexToAdjacentFaces, int32 FaceIndex, TArray<int32>& TriangleQueue, const bool bConnectByEdge)
+	{
+		int32 NumFaces = (int32)BuildData->GetNumFaces();
+		check(FaceAdded.Num() == NumFaces);
+
+		TMap<int32, int32> AdjacentFaceCommonVertices;
+		for (int32 Corner = 0; Corner < 3; Corner++)
+		{
+			int32 VertexIndex = BuildData->GetVertexIndex(FaceIndex, Corner);
+			TArray<int32>& AdjacentFaces = VertexIndexToAdjacentFaces.FindChecked(VertexIndex);
+			for (int32 AdjacentFaceArrayIndex = 0; AdjacentFaceArrayIndex < AdjacentFaces.Num(); ++AdjacentFaceArrayIndex)
+			{
+				int32 AdjacentFaceIndex = AdjacentFaces[AdjacentFaceArrayIndex];
+				if (!FaceAdded[AdjacentFaceIndex] && AdjacentFaceIndex != FaceIndex)
+				{
+					bool bAddConnected = !bConnectByEdge;
+					if (bConnectByEdge)
+					{
+						int32& AdjacentFaceCommonVerticeCount = AdjacentFaceCommonVertices.FindOrAdd(AdjacentFaceIndex);
+						AdjacentFaceCommonVerticeCount++;
+						//Is the connected triangles share 2 vertex index (one edge) not only one vertex
+						bAddConnected = AdjacentFaceCommonVerticeCount > 1;
+					}
+
+					if (bAddConnected)
+					{
+						TriangleQueue.Add(AdjacentFaceIndex);
+						//Add the face only once by marking the face has computed
+						FaceAdded[AdjacentFaceIndex] = true;
+					}
+				}
+			}
+		}
+	}
+
+	//Fill FaceIndexToPatchIndex so every triangle know is unique island patch index.
+	//We need to respect the island when we use the smooth group to compute the normals.
+	//Each island patch have its own smoothgroup data, there is no triangle connectivity possible between island patch.
+	//@Param bConnectByEdge: If true we need at least 2 vertex index (one edge) to connect 2 triangle. If false we just need one vertex index (bowtie)
+	void Skeletal_FillPolygonPatch(IMeshBuildData* BuildData, TArray<int32>& FaceIndexToPatchIndex, const bool bConnectByEdge)
+	{
+		int32 NumTriangles = BuildData->GetNumFaces();
+		check(FaceIndexToPatchIndex.Num() == NumTriangles);
+		
+		int32 PatchIndex = 0;
+
+		TMap<int32, TArray<int32>> VertexIndexToAdjacentFaces;
+		VertexIndexToAdjacentFaces.Reserve(BuildData->GetNumFaces()*2);
+		for (int32 FaceIndex = 0; FaceIndex < NumTriangles; ++FaceIndex)
+		{
+			int32 WedgeOffset = FaceIndex * 3;
+			for (int32 Corner = 0; Corner < 3; Corner++)
+			{
+				int32 VertexIndex = BuildData->GetVertexIndex(FaceIndex, Corner);
+				TArray<int32>& AdjacentFaces = VertexIndexToAdjacentFaces.FindOrAdd(VertexIndex);
+				AdjacentFaces.AddUnique(FaceIndex);
+			}
+		}
+
+		//Mark added face so we do not add them more then once
+		TBitArray<> FaceAdded;
+		FaceAdded.Init(false, NumTriangles);
+
+		TArray<int32> TriangleQueue;
+		TriangleQueue.Reserve(100);
+		for (int32 FaceIndex = 0; FaceIndex < NumTriangles; ++FaceIndex)
+		{
+			if (FaceAdded[FaceIndex])
+			{
+				continue;
+			}
+			TriangleQueue.Reset();
+			TriangleQueue.Add(FaceIndex); //Use a queue to avoid recursive function
+			FaceAdded[FaceIndex] = true;
+			while (TriangleQueue.Num() > 0)
+			{
+				int32 CurrentTriangleIndex = TriangleQueue.Pop(false);
+				FaceIndexToPatchIndex[CurrentTriangleIndex] = PatchIndex;
+				AddAdjacentFace(BuildData, FaceAdded, VertexIndexToAdjacentFaces, CurrentTriangleIndex, TriangleQueue, bConnectByEdge);
+			}
+			PatchIndex++;
+		}
+	}
+
 	void Skeletal_ComputeTangents(
 		IMeshBuildData* BuildData,
 		const FOverlappingCorners& OverlappingCorners
@@ -3665,7 +3752,11 @@ public:
 	{
 		bool bBlendOverlappingNormals = true;
 		bool bIgnoreDegenerateTriangles = BuildData->BuildOptions.bRemoveDegenerateTriangles;
-
+		
+		int32 NumFaces = BuildData->GetNumFaces();
+		int32 NumWedges = BuildData->GetNumWedges();
+		check(NumFaces * 3 == NumWedges);
+		
 		// Compute per-triangle tangents.
 		TArray<FVector> TriangleTangentX;
 		TArray<FVector> TriangleTangentY;
@@ -3679,6 +3770,12 @@ public:
 			bIgnoreDegenerateTriangles ? SMALL_NUMBER : 0.0f
 			);
 
+		TArray<int32> FaceIndexToPatchIndex;
+		FaceIndexToPatchIndex.AddZeroed(NumFaces);
+		//Since we use triangle normals to compute the vertex normal, we need a full edge connected (2 vertex component per triangle)
+		const bool bConnectByEdge = true;
+		Skeletal_FillPolygonPatch(BuildData, FaceIndexToPatchIndex, bConnectByEdge);
+
 		TArray<FVector>& WedgeTangentX = BuildData->GetTangentArray(0);
 		TArray<FVector>& WedgeTangentY = BuildData->GetTangentArray(1);
 		TArray<FVector>& WedgeTangentZ = BuildData->GetTangentArray(2);
@@ -3687,10 +3784,6 @@ public:
 		TArray<FFanFace> RelevantFacesForCorner[3];
 		TArray<int32> AdjacentFaces;
 
-		int32 NumFaces = BuildData->GetNumFaces();
-		int32 NumWedges = BuildData->GetNumWedges();
-		check(NumFaces * 3 == NumWedges);
-
 		bool bWedgeTSpace = false;
 
 		if (WedgeTangentX.Num() > 0 && WedgeTangentY.Num() > 0)
@@ -3715,6 +3808,7 @@ public:
 
 		for (int32 FaceIndex = 0; FaceIndex < NumFaces; FaceIndex++)
 		{
+			int32 PatchIndex = FaceIndexToPatchIndex[FaceIndex];
 			int32 WedgeOffset = FaceIndex * 3;
 			FVector CornerPositions[3];
 			FVector CornerNormal[3];
@@ -3758,8 +3852,9 @@ public:
 				for (int32 k = 0; k < DupVerts.Num(); k++)
 				{
 					int32 PotentialTriangleIndex = DupVerts[k] / 3;
-					//Do not add mirror triangle to the adjacentFaces
-					if (!IsTriangleMirror(BuildData, TriangleTangentZ, FaceIndex, PotentialTriangleIndex))
+					
+					//Do not add mirror triangle to the adjacentFaces. Also make sure adjacent triangle is in the same connected triangle patch.
+					if (!IsTriangleMirror(BuildData, TriangleTangentZ, FaceIndex, PotentialTriangleIndex) && PatchIndex == FaceIndexToPatchIndex[PotentialTriangleIndex])
 					{
 						AdjacentFaces.AddUnique(PotentialTriangleIndex);
 					}
diff --git a/Engine/Source/Editor/UnrealEd/Private/Fbx/FbxSkeletalMeshImport.cpp b/Engine/Source/Editor/UnrealEd/Private/Fbx/FbxSkeletalMeshImport.cpp
--- a/Engine/Source/Editor/UnrealEd/Private/Fbx/FbxSkeletalMeshImport.cpp
+++ b/Engine/Source/Editor/UnrealEd/Private/Fbx/FbxSkeletalMeshImport.cpp
@@ -3174,7 +3174,11 @@ bool UnFbx::FFbxImporter::FillSkelMeshImporterFromFbx( FSkeletalMeshImportData&
 	int32 LayerSmoothingCount = Mesh->GetLayerCount(FbxLayerElement::eSmoothing);
 	for(int32 i = 0; i < LayerSmoothingCount; i++)
 	{
-		GeometryConverter->ComputePolygonSmoothingFromEdgeSmoothing (Mesh, i);
+		FbxLayerElementSmoothing const* SmoothingInfo = Mesh->GetLayer(i)->GetSmoothing();
+		if (SmoothingInfo && SmoothingInfo->GetMappingMode() != FbxLayerElement::eByPolygon)
+		{
+			GeometryConverter->ComputePolygonSmoothingFromEdgeSmoothing(Mesh, i);
+		}
 	}
 
 	//
@@ -3272,6 +3276,12 @@ bool UnFbx::FFbxImporter::FillSkelMeshImporterFromFbx( FSkeletalMeshImportData&
 				AddTokenizedErrorMessage(FTokenizedMessage::Create(EMessageSeverity::Warning, FText::Format(LOCTEXT("FbxSkeletaLMeshimport_ConvertSmoothingGroupFailed", "Unable to fully convert the smoothing groups for mesh '{0}'"), FText::FromString(Mesh->GetName()))), FFbxErrors::Generic_Mesh_ConvertSmoothingGroupFailed);
 				bSmoothingAvailable = false;
 			}
+			else
+			{
+				//After using the geometry converter we always have to get the Layer and the smoothing info
+				BaseLayer = Mesh->GetLayer(0);
+				SmoothingInfo = BaseLayer->GetSmoothing();
+			}
 		}
 
 		if( SmoothingInfo->GetMappingMode() == FbxLayerElement::eByPolygon )
