Add a fail safe to the old skeletal mesh legacy compute tangent.

diff --git a/Engine/Source/Developer/MeshUtilities/Private/MeshUtilities.cpp b/Engine/Source/Developer/MeshUtilities/Private/MeshUtilities.cpp
--- a/Engine/Source/Developer/MeshUtilities/Private/MeshUtilities.cpp
+++ b/Engine/Source/Developer/MeshUtilities/Private/MeshUtilities.cpp
@@ -4364,6 +4364,90 @@ bool FMeshUtilities::BuildSkeletalMesh(FSkeletalMeshLODModel& LODModel, const FR
 #endif
 }
 
+//The fail safe is there to avoid zeros in the tangents. Even if the fail safe prevent zero NTBs, a warning should be generate by the caller to let the artist know something went wrong
+//Using a fail safe can lead to hard edge where its suppose to be smooth, it can also have some impact on the shading (lighting for tangentZ and normal map for tangentX and Y)
+//Normally because we use the triangle data the tangent space is in a good direction and should give proper result.
+void TangentFailSafe(const FVector &TriangleTangentX, const FVector &TriangleTangentY, const FVector &TriangleTangentZ
+	, FVector &TangentX, FVector &TangentY, FVector &TangentZ)
+{
+	bool bTangentXZero = TangentX.IsNearlyZero() || TangentX.ContainsNaN();
+	bool bTangentYZero = TangentY.IsNearlyZero() || TangentY.ContainsNaN();
+	bool bTangentZZero = TangentZ.IsNearlyZero() || TangentZ.ContainsNaN();
+
+	if (!bTangentXZero && !bTangentYZero && !bTangentZZero)
+	{
+		//No need to fail safe if everything is different from zero
+		return;
+	}
+	if (!bTangentZZero)
+	{
+		if (!bTangentXZero)
+		{
+			//Valid TangentZ and TangentX, we can recompute TangentY
+			TangentY = FVector::CrossProduct(TangentZ, TangentX).GetSafeNormal();
+		}
+		else if (!bTangentYZero)
+		{
+			//Valid TangentZ and TangentY, we can recompute TangentX
+			TangentX = FVector::CrossProduct(TangentY, TangentZ).GetSafeNormal();
+		}
+		else
+		{
+			//TangentX and Y are invalid, use the triangle data, can cause a hard edge
+			TangentX = TriangleTangentX;
+			TangentY = TriangleTangentY;
+		}
+	}
+	else if (!bTangentXZero)
+	{
+		if (!bTangentYZero)
+		{
+			//Valid TangentX and TangentY, we can recompute TangentZ
+			TangentZ = FVector::CrossProduct(TangentX, TangentY).GetSafeNormal();
+		}
+		else
+		{
+			//TangentY and Z are invalid, use the triangle data, can cause a hard edge
+			TangentZ = TriangleTangentZ;
+			TangentY = TriangleTangentY;
+		}
+	}
+	else if (!bTangentYZero)
+	{
+		//TangentX and Z are invalid, use the triangle data, can cause a hard edge
+		TangentX = TriangleTangentX;
+		TangentZ = TriangleTangentZ;
+	}
+	else
+	{
+		//Everything is zero, use all triangle data, can cause a hard edge
+		TangentX = TriangleTangentX;
+		TangentY = TriangleTangentY;
+		TangentZ = TriangleTangentZ;
+	}
+
+	//Ortho normalize the result
+	TangentY -= TangentX * (TangentX | TangentY);
+	TangentY.Normalize();
+
+	TangentX -= TangentZ * (TangentZ | TangentX);
+	TangentY -= TangentZ * (TangentZ | TangentY);
+
+	TangentX.Normalize();
+	TangentY.Normalize();
+
+	//If we still have some zero data (i.e. triangle data is degenerated)
+	if (TangentZ.IsNearlyZero() || TangentZ.ContainsNaN()
+		|| TangentX.IsNearlyZero() || TangentX.ContainsNaN()
+		|| TangentY.IsNearlyZero() || TangentY.ContainsNaN())
+	{
+		//Since the triangle is degenerate this case can cause a hardedge, but will probably have no other impact since the triangle is degenerate (no visible surface)
+		TangentX = FVector(1.0f, 0.0f, 0.0f);
+		TangentY = FVector(0.0f, 1.0f, 0.0f);
+		TangentZ = FVector(0.0f, 0.0f, 1.0f);
+	}
+}
+
 //@TODO: The OutMessages has to be a struct that contains FText/FName, or make it Token and add that as error. Needs re-work. Temporary workaround for now. 
 bool FMeshUtilities::BuildSkeletalMesh_Legacy(FSkeletalMeshLODModel& LODModel
 											, const FReferenceSkeleton& RefSkeleton
@@ -4691,17 +4775,21 @@ bool FMeshUtilities::BuildSkeletalMesh_Legacy(FSkeletalMeshLODModel& LODModel
 				TangentZ.Normalize();
 			}
 
-			if (TangentX.IsNearlyZero() || TangentX.ContainsNaN())
-			{
-				NTBErrorCount++;
-			}
-			if (TangentY.IsNearlyZero() || TangentY.ContainsNaN())
-			{
-				NTBErrorCount++;
-			}
-			if (TangentZ.IsNearlyZero() || TangentZ.ContainsNaN())
+			//FAIL safe, avoid zero tangents
+			bool bTangentXZero = TangentX.IsNearlyZero() || TangentX.ContainsNaN();
+			bool bTangentYZero = TangentY.IsNearlyZero() || TangentY.ContainsNaN();
+			bool bTangentZZero = TangentZ.IsNearlyZero() || TangentZ.ContainsNaN();
+			if (bTangentXZero || bTangentYZero || bTangentZZero)
 			{
 				NTBErrorCount++;
+				FVector TriangleTangentZ = FPlane(
+					Points[Wedges[Face.iWedge[2]].iVertex],
+					Points[Wedges[Face.iWedge[1]].iVertex],
+					Points[Wedges[Face.iWedge[0]].iVertex]
+				);
+				FVector TriangleTangentX = FaceTangentX[FaceIndex];
+				FVector TriangleTangentY = FaceTangentY[FaceIndex];
+				TangentFailSafe(TriangleTangentX, TriangleTangentY, TriangleTangentZ, TangentX, TangentY, TangentZ);
 			}
 
 			Vertex.TangentX = TangentX;
