Fix various issues related to lightmap 2-texel padding not being considered

* UVs in lightmap density view were shifting after building lighting. Fixed by considering the padding even if lightmap data is invalid
* Auto generated lightmap UV was snapping charts to the whole padded lightmap instead of the actual valid center region. Fixed by subtracing 2 texels from the packing resolution
* Also fixed some merging conflicts of auto gen

diff --git a/Engine/Source/Developer/MeshUtilities/Private/MeshUtilities.cpp b/Engine/Source/Developer/MeshUtilities/Private/MeshUtilities.cpp
--- a/Engine/Source/Developer/MeshUtilities/Private/MeshUtilities.cpp
+++ b/Engine/Source/Developer/MeshUtilities/Private/MeshUtilities.cpp
@@ -2507,7 +2507,17 @@ public:
 					Packer.SetVersion(LightmapUVVersion);
 
 					Packer.FindCharts(OverlappingCorners);
-					bool bPackSuccess = Packer.FindBestPacking(SrcModel.BuildSettings.MinLightmapResolution);
+
+					int32 EffectiveMinLightmapResolution = SrcModel.BuildSettings.MinLightmapResolution;
+					if (LightmapUVVersion >= ELightmapUVVersion::ConsiderLightmapPadding)
+					{
+						if (GLightmassDebugOptions.bPadMappings)
+						{
+							EffectiveMinLightmapResolution -= 2;
+						}
+					}
+
+					bool bPackSuccess = Packer.FindBestPacking(EffectiveMinLightmapResolution);
 					if (bPackSuccess)
 					{
 						Packer.CommitPackedUVs();
diff --git a/Engine/Source/Developer/MeshUtilitiesCommon/Public/MeshUtilitiesCommon.h b/Engine/Source/Developer/MeshUtilitiesCommon/Public/MeshUtilitiesCommon.h
--- a/Engine/Source/Developer/MeshUtilitiesCommon/Public/MeshUtilitiesCommon.h
+++ b/Engine/Source/Developer/MeshUtilitiesCommon/Public/MeshUtilitiesCommon.h
@@ -12,7 +12,8 @@ enum class ELightmapUVVersion : int32
 	ScaleChartsOrderingFix = 3,
 	ChartJoiningLFix = 4,
 	Allocator2DFlipFix = 5,
-	Latest = ScaleChartsOrderingFix
+	ConsiderLightmapPadding = 6,
+	Latest = ConsiderLightmapPadding
 };
 
 /** Helper struct for building acceleration structures. */
diff --git a/Engine/Source/Runtime/Renderer/Private/LightMapDensityRendering.cpp b/Engine/Source/Runtime/Renderer/Private/LightMapDensityRendering.cpp
--- a/Engine/Source/Runtime/Renderer/Private/LightMapDensityRendering.cpp
+++ b/Engine/Source/Runtime/Renderer/Private/LightMapDensityRendering.cpp
@@ -174,6 +174,12 @@ void FLightmapDensityMeshProcessor::Process(
 		else if (PrimitiveSceneProxy)
 		{
 			int32 LightMapResolution = PrimitiveSceneProxy->GetLightMapResolution();
+		#if WITH_EDITOR
+			if (GLightmassDebugOptions.bPadMappings)
+			{
+				LightMapResolution -= 2;
+			}
+		#endif
 			if (PrimitiveSceneProxy->IsStatic() && LightMapResolution > 0)
 			{
 				ShaderElementData.bTextureMapped = true;
