Adding more null checks for LODModel.AdditionalIndexBuffers - we haven't seen a crash originating here yet, but there's potential for one

diff --git a/Engine/Source/Developer/MeshUtilities/Private/MeshUtilities.cpp b/Engine/Source/Developer/MeshUtilities/Private/MeshUtilities.cpp
--- a/Engine/Source/Developer/MeshUtilities/Private/MeshUtilities.cpp
+++ b/Engine/Source/Developer/MeshUtilities/Private/MeshUtilities.cpp
@@ -2785,7 +2785,7 @@ public:
 			LODModel.IndexBuffer.SetIndices(CombinedIndices, bNeeds32BitIndices ? EIndexBufferStride::Force32Bit : EIndexBufferStride::Force16Bit);
 			
 			// Build the reversed index buffer.
-			if (InOutModels[0].BuildSettings.bBuildReversedIndexBuffer)
+			if (LODModel.AdditionalIndexBuffers && InOutModels[0].BuildSettings.bBuildReversedIndexBuffer)
 			{
 				TArray<uint32> InversedIndices;
 				const int32 IndexCount = CombinedIndices.Num();
@@ -2823,7 +2823,7 @@ public:
 			}
 
 			// Build the inversed depth only index buffer.
-			if (InOutModels[0].BuildSettings.bBuildReversedIndexBuffer)
+			if (LODModel.AdditionalIndexBuffers && InOutModels[0].BuildSettings.bBuildReversedIndexBuffer)
 			{
 				TArray<uint32> ReversedDepthOnlyIndices;
 				const int32 IndexCount = DepthOnlyIndices.Num();
@@ -2836,6 +2836,7 @@ public:
 			}
 
 			// Build a list of wireframe edges in the static mesh.
+			if (LODModel.AdditionalIndexBuffers)
 			{
 				TArray<FMeshEdgeDef> Edges;
 				TArray<uint32> WireframeIndices;
@@ -2852,7 +2853,7 @@ public:
 			}
 
 			// Build the adjacency index buffer used for tessellation.
-			if (InOutModels[0].BuildSettings.bBuildAdjacencyBuffer)
+			if (LODModel.AdditionalIndexBuffers && InOutModels[0].BuildSettings.bBuildAdjacencyBuffer)
 			{
 				TArray<uint32> AdjacencyIndices;
 
