Fix degenerated triangle that introduce nan value into the normals when we compute them

diff --git a/Engine/Source/Developer/MeshUtilities/Private/MeshUtilities.cpp b/Engine/Source/Developer/MeshUtilities/Private/MeshUtilities.cpp
--- a/Engine/Source/Developer/MeshUtilities/Private/MeshUtilities.cpp
+++ b/Engine/Source/Developer/MeshUtilities/Private/MeshUtilities.cpp
@@ -1056,6 +1056,9 @@ static void ComputeTriangleTangents(
 	OutTangentY.Empty(NumTriangles);
 	OutTangentZ.Empty(NumTriangles);
 
+	//Currently GetSafeNormal do not support 0.0f threshold properly
+	float RealComparisonThreshold = FMath::Max(ComparisonThreshold, FLT_MIN);
+
 	for (int32 TriangleIndex = 0; TriangleIndex < NumTriangles; TriangleIndex++)
 	{
 		int32 UVIndex = 0;
@@ -1066,37 +1069,57 @@ static void ComputeTriangleTangents(
 			P[i] = InVertices[InIndices[TriangleIndex * 3 + i]];
 		}
 
-		const FVector Normal = ((P[1] - P[2]) ^ (P[0] - P[2])).GetSafeNormal(ComparisonThreshold);
-		FMatrix	ParameterToLocal(
-			FPlane(P[1].X - P[0].X, P[1].Y - P[0].Y, P[1].Z - P[0].Z, 0),
-			FPlane(P[2].X - P[0].X, P[2].Y - P[0].Y, P[2].Z - P[0].Z, 0),
-			FPlane(P[0].X, P[0].Y, P[0].Z, 0),
-			FPlane(0, 0, 0, 1)
+		const FVector Normal = ((P[1] - P[2]) ^ (P[0] - P[2])).GetSafeNormal(RealComparisonThreshold);
+		//Avoid doing orthonormal vector from a degenerated triangle.
+		if (!Normal.IsNearlyZero(FLT_MIN))
+		{
+			FMatrix	ParameterToLocal(
+				FPlane(P[1].X - P[0].X, P[1].Y - P[0].Y, P[1].Z - P[0].Z, 0),
+				FPlane(P[2].X - P[0].X, P[2].Y - P[0].Y, P[2].Z - P[0].Z, 0),
+				FPlane(P[0].X, P[0].Y, P[0].Z, 0),
+				FPlane(0, 0, 0, 1)
 			);
 
-		const FVector2D T1 = InUVs[TriangleIndex * 3 + 0];
-		const FVector2D T2 = InUVs[TriangleIndex * 3 + 1];
-		const FVector2D T3 = InUVs[TriangleIndex * 3 + 2];
+			const FVector2D T1 = InUVs[TriangleIndex * 3 + 0];
+			const FVector2D T2 = InUVs[TriangleIndex * 3 + 1];
+			const FVector2D T3 = InUVs[TriangleIndex * 3 + 2];
 
-		FMatrix ParameterToTexture(
-			FPlane(T2.X - T1.X, T2.Y - T1.Y, 0, 0),
-			FPlane(T3.X - T1.X, T3.Y - T1.Y, 0, 0),
-			FPlane(T1.X, T1.Y, 1, 0),
-			FPlane(0, 0, 0, 1)
+			FMatrix ParameterToTexture(
+				FPlane(T2.X - T1.X, T2.Y - T1.Y, 0, 0),
+				FPlane(T3.X - T1.X, T3.Y - T1.Y, 0, 0),
+				FPlane(T1.X, T1.Y, 1, 0),
+				FPlane(0, 0, 0, 1)
 			);
 
-		// Use InverseSlow to catch singular matrices.  Inverse can miss this sometimes.
-		const FMatrix TextureToLocal = ParameterToTexture.Inverse() * ParameterToLocal;
+			// Use InverseSlow to catch singular matrices.  Inverse can miss this sometimes.
+			const FMatrix TextureToLocal = ParameterToTexture.Inverse() * ParameterToLocal;
 
-		OutTangentX.Add(TextureToLocal.TransformVector(FVector(1, 0, 0)).GetSafeNormal());
-		OutTangentY.Add(TextureToLocal.TransformVector(FVector(0, 1, 0)).GetSafeNormal());
-		OutTangentZ.Add(Normal);
+			OutTangentX.Add(TextureToLocal.TransformVector(FVector(1, 0, 0)).GetSafeNormal());
+			OutTangentY.Add(TextureToLocal.TransformVector(FVector(0, 1, 0)).GetSafeNormal());
+			OutTangentZ.Add(Normal);
 
-		FVector::CreateOrthonormalBasis(
-			OutTangentX[TriangleIndex],
-			OutTangentY[TriangleIndex],
-			OutTangentZ[TriangleIndex]
+			FVector::CreateOrthonormalBasis(
+				OutTangentX[TriangleIndex],
+				OutTangentY[TriangleIndex],
+				OutTangentZ[TriangleIndex]
 			);
+			if (OutTangentX[TriangleIndex].IsNearlyZero() || OutTangentX[TriangleIndex].ContainsNaN()
+				|| OutTangentY[TriangleIndex].IsNearlyZero() || OutTangentY[TriangleIndex].ContainsNaN()
+				|| OutTangentZ[TriangleIndex].IsNearlyZero() || OutTangentZ[TriangleIndex].ContainsNaN())
+			{
+				OutTangentX[TriangleIndex] = FVector::ZeroVector;
+				OutTangentY[TriangleIndex] = FVector::ZeroVector;
+				OutTangentZ[TriangleIndex] = FVector::ZeroVector;
+			}
+		}
+		else
+		{
+			//Add zero tangents and normal for this triangle, this is like weighting it to zero when we compute the vertex normal
+			//But we need the triangle to correctly connect other neighbourg triangles
+			OutTangentX.Add(FVector::ZeroVector);
+			OutTangentY.Add(FVector::ZeroVector);
+			OutTangentZ.Add(FVector::ZeroVector);
+		}
 	}
 
 	check(OutTangentX.Num() == NumTriangles);
@@ -1237,7 +1260,7 @@ static void ComputeTangents(
 		TriangleTangentX,
 		TriangleTangentY,
 		TriangleTangentZ,
-		bIgnoreDegenerateTriangles ? SMALL_NUMBER : 0.0f
+		bIgnoreDegenerateTriangles ? SMALL_NUMBER : FLT_MIN
 		);
 
 	// Declare these out here to avoid reallocations.
@@ -1773,7 +1796,7 @@ static void ComputeNormals(
 		TriangleTangentX,
 		TriangleTangentY,
 		TriangleTangentZ,
-		bIgnoreDegenerateTriangles ? SMALL_NUMBER : 0.0f
+		bIgnoreDegenerateTriangles ? SMALL_NUMBER : FLT_MIN
 		);
 
 	// Declare these out here to avoid reallocations.
@@ -3266,6 +3289,9 @@ public:
 		TriangleTangentY.Empty(NumTriangles);
 		TriangleTangentZ.Empty(NumTriangles);
 
+		//Currently GetSafeNormal do not support 0.0f threshold properly
+		float RealComparisonThreshold = FMath::Max(ComparisonThreshold, FLT_MIN);
+
 		for (int32 TriangleIndex = 0; TriangleIndex < NumTriangles; TriangleIndex++)
 		{
 			const int32 UVIndex = 0;
@@ -3276,36 +3302,58 @@ public:
 				P[i] = BuildData->GetVertexPosition(TriangleIndex, i);
 			}
 
-			const FVector Normal = ((P[1] - P[2]) ^ (P[0] - P[2])).GetSafeNormal(ComparisonThreshold);
-			FMatrix	ParameterToLocal(
-				FPlane(P[1].X - P[0].X, P[1].Y - P[0].Y, P[1].Z - P[0].Z, 0),
-				FPlane(P[2].X - P[0].X, P[2].Y - P[0].Y, P[2].Z - P[0].Z, 0),
-				FPlane(P[0].X, P[0].Y, P[0].Z, 0),
-				FPlane(0, 0, 0, 1)
+			//get safe normal should have return a valid normalized vector or a zero vector.
+			const FVector Normal = ((P[1] - P[2]) ^ (P[0] - P[2])).GetSafeNormal(RealComparisonThreshold);
+			//Avoid doing orthonormal vector from a degenerated triangle.
+			if (!Normal.IsNearlyZero(FLT_MIN))
+			{
+				FMatrix	ParameterToLocal(
+					FPlane(P[1].X - P[0].X, P[1].Y - P[0].Y, P[1].Z - P[0].Z, 0),
+					FPlane(P[2].X - P[0].X, P[2].Y - P[0].Y, P[2].Z - P[0].Z, 0),
+					FPlane(P[0].X, P[0].Y, P[0].Z, 0),
+					FPlane(0, 0, 0, 1)
 				);
 
-			FVector2D T1 = BuildData->GetVertexUV(TriangleIndex, 0, UVIndex);
-			FVector2D T2 = BuildData->GetVertexUV(TriangleIndex, 1, UVIndex);
-			FVector2D T3 = BuildData->GetVertexUV(TriangleIndex, 2, UVIndex);
-			FMatrix ParameterToTexture(
-				FPlane(T2.X - T1.X, T2.Y - T1.Y, 0, 0),
-				FPlane(T3.X - T1.X, T3.Y - T1.Y, 0, 0),
-				FPlane(T1.X, T1.Y, 1, 0),
-				FPlane(0, 0, 0, 1)
+				FVector2D T1 = BuildData->GetVertexUV(TriangleIndex, 0, UVIndex);
+				FVector2D T2 = BuildData->GetVertexUV(TriangleIndex, 1, UVIndex);
+				FVector2D T3 = BuildData->GetVertexUV(TriangleIndex, 2, UVIndex);
+				FMatrix ParameterToTexture(
+					FPlane(T2.X - T1.X, T2.Y - T1.Y, 0, 0),
+					FPlane(T3.X - T1.X, T3.Y - T1.Y, 0, 0),
+					FPlane(T1.X, T1.Y, 1, 0),
+					FPlane(0, 0, 0, 1)
 				);
 
-			// Use InverseSlow to catch singular matrices.  Inverse can miss this sometimes.
-			const FMatrix TextureToLocal = ParameterToTexture.Inverse() * ParameterToLocal;
+				// Use InverseSlow to catch singular matrices.  Inverse can miss this sometimes.
+				const FMatrix TextureToLocal = ParameterToTexture.Inverse() * ParameterToLocal;
 
-			TriangleTangentX.Add(TextureToLocal.TransformVector(FVector(1, 0, 0)).GetSafeNormal());
-			TriangleTangentY.Add(TextureToLocal.TransformVector(FVector(0, 1, 0)).GetSafeNormal());
-			TriangleTangentZ.Add(Normal);
+				TriangleTangentX.Add(TextureToLocal.TransformVector(FVector(1, 0, 0)).GetSafeNormal());
+				TriangleTangentY.Add(TextureToLocal.TransformVector(FVector(0, 1, 0)).GetSafeNormal());
+				TriangleTangentZ.Add(Normal);
 
-			FVector::CreateOrthonormalBasis(
-				TriangleTangentX[TriangleIndex],
-				TriangleTangentY[TriangleIndex],
-				TriangleTangentZ[TriangleIndex]
+				FVector::CreateOrthonormalBasis(
+					TriangleTangentX[TriangleIndex],
+					TriangleTangentY[TriangleIndex],
+					TriangleTangentZ[TriangleIndex]
 				);
+
+				if (TriangleTangentX[TriangleIndex].IsNearlyZero() || TriangleTangentX[TriangleIndex].ContainsNaN()
+					|| TriangleTangentY[TriangleIndex].IsNearlyZero() || TriangleTangentY[TriangleIndex].ContainsNaN()
+					|| TriangleTangentZ[TriangleIndex].IsNearlyZero() || TriangleTangentZ[TriangleIndex].ContainsNaN())
+				{
+					TriangleTangentX[TriangleIndex] = FVector::ZeroVector;
+					TriangleTangentY[TriangleIndex] = FVector::ZeroVector;
+					TriangleTangentZ[TriangleIndex] = FVector::ZeroVector;
+				}
+			}
+			else
+			{
+				//Add zero tangents and normal for this triangle, this is like weighting it to zero when we compute the vertex normal
+				//But we need the triangle to correctly connect other neighbourg triangles
+				TriangleTangentX.Add(FVector::ZeroVector);
+				TriangleTangentY.Add(FVector::ZeroVector);
+				TriangleTangentZ.Add(FVector::ZeroVector);
+			}
 		}
 	}
 
@@ -3414,7 +3462,7 @@ public:
 			TriangleTangentY,
 			TriangleTangentZ,
 			BuildData,
-			bIgnoreDegenerateTriangles ? SMALL_NUMBER : 0.0f
+			bIgnoreDegenerateTriangles ? SMALL_NUMBER : FLT_MIN
 			);
 
 		TArray<FVector>& WedgeTangentX = BuildData->GetTangentArray(0);
@@ -3774,7 +3822,7 @@ public:
 			TriangleTangentY,
 			TriangleTangentZ,
 			BuildData,
-			bIgnoreDegenerateTriangles ? SMALL_NUMBER : 0.0f
+			bIgnoreDegenerateTriangles ? SMALL_NUMBER : FLT_MIN
 			);
 
 		TArray<int32> FaceIndexToPatchIndex;
@@ -3810,9 +3858,8 @@ public:
 			WedgeTangentZ.Empty(NumWedges);
 			WedgeTangentZ.AddZeroed(NumWedges);
 		}
-
+		bool bIsZeroLengthNormalErrorMessageDisplayed = false;
 		// we need to calculate normals for MikkTSpace
-
 		for (int32 FaceIndex = 0; FaceIndex < NumFaces; FaceIndex++)
 		{
 			int32 PatchIndex = FaceIndexToPatchIndex[FaceIndex];
@@ -3860,8 +3907,9 @@ public:
 				{
 					int32 PotentialTriangleIndex = DupVerts[k] / 3;
 					
-					//Do not add mirror triangle to the adjacentFaces. Also make sure adjacent triangle is in the same connected triangle patch.
-					if (!IsTriangleMirror(BuildData, TriangleTangentZ, FaceIndex, PotentialTriangleIndex) && PatchIndex == FaceIndexToPatchIndex[PotentialTriangleIndex])
+					bool bDegeneratedTriangles = TriangleTangentZ[FaceIndex].IsNearlyZero() || TriangleTangentZ[PotentialTriangleIndex].IsNearlyZero();
+					//Do not add mirror triangle to the adjacentFaces. Also make sure adjacent triangle is in the same connected triangle patch. Accept connected degenerate triangle
+					if ((bDegeneratedTriangles || !IsTriangleMirror(BuildData, TriangleTangentZ, FaceIndex, PotentialTriangleIndex)) && PatchIndex == FaceIndexToPatchIndex[PotentialTriangleIndex])
 					{
 						AdjacentFaces.AddUnique(PotentialTriangleIndex);
 					}
@@ -4019,6 +4067,30 @@ public:
 			// Copy back to the mesh.
 			for (int32 CornerIndex = 0; CornerIndex < 3; CornerIndex++)
 			{
+				//Make sure the normal do not contain nan, if it contain nan change it to a valid zero vector
+				if (CornerNormal[CornerIndex].IsNearlyZero() || CornerNormal[CornerIndex].ContainsNaN())
+				{
+					CornerNormal[CornerIndex] = FVector::ZeroVector;
+					//We also notify the log that we compute a zero length normals, so the user is aware of it
+					if (!bIsZeroLengthNormalErrorMessageDisplayed)
+					{
+						bIsZeroLengthNormalErrorMessageDisplayed = true;
+						// add warning message if available, do a log if not
+						FText TextMessage = LOCTEXT("Skeletal_ComputeTangents_MikkTSpace_Warning_ZeroLengthNormal", "Skeletal ComputeTangents MikkTSpace function: Compute a zero length normal vector.");
+						if (BuildData->OutWarningMessages)
+						{
+							BuildData->OutWarningMessages->Add(TextMessage);
+							if (BuildData->OutWarningNames)
+							{
+								BuildData->OutWarningNames->Add(FFbxErrors::Generic_Mesh_TangentsComputeError);
+							}
+						}
+						else
+						{
+							UE_LOG(LogSkeletalMesh, Warning, TEXT("%s"), *(TextMessage.ToString()));
+						}
+					}
+				}
 				WedgeTangentZ[WedgeOffset + CornerIndex] = CornerNormal[CornerIndex];
 			}
 		}
@@ -4622,37 +4694,45 @@ bool FMeshUtilities::BuildSkeletalMesh_Legacy(FSkeletalMeshLODModel& LODModel
 				P2 = Points[Wedges[Faces[FaceIndex].iWedge[1]].iVertex],
 				P3 = Points[Wedges[Faces[FaceIndex].iWedge[2]].iVertex];
 			FVector	TriangleNormal = FPlane(P3, P2, P1);
-			FMatrix	ParameterToLocal(
-				FPlane(P2.X - P1.X, P2.Y - P1.Y, P2.Z - P1.Z, 0),
-				FPlane(P3.X - P1.X, P3.Y - P1.Y, P3.Z - P1.Z, 0),
-				FPlane(P1.X, P1.Y, P1.Z, 0),
-				FPlane(0, 0, 0, 1)
+			if (!TriangleNormal.IsNearlyZero(FLT_MIN))
+			{
+				FMatrix	ParameterToLocal(
+					FPlane(P2.X - P1.X, P2.Y - P1.Y, P2.Z - P1.Z, 0),
+					FPlane(P3.X - P1.X, P3.Y - P1.Y, P3.Z - P1.Z, 0),
+					FPlane(P1.X, P1.Y, P1.Z, 0),
+					FPlane(0, 0, 0, 1)
 				);
 
-			float	U1 = Wedges[Faces[FaceIndex].iWedge[0]].UVs[0].X,
-				U2 = Wedges[Faces[FaceIndex].iWedge[1]].UVs[0].X,
-				U3 = Wedges[Faces[FaceIndex].iWedge[2]].UVs[0].X,
-				V1 = Wedges[Faces[FaceIndex].iWedge[0]].UVs[0].Y,
-				V2 = Wedges[Faces[FaceIndex].iWedge[1]].UVs[0].Y,
-				V3 = Wedges[Faces[FaceIndex].iWedge[2]].UVs[0].Y;
-
-			FMatrix	ParameterToTexture(
-				FPlane(U2 - U1, V2 - V1, 0, 0),
-				FPlane(U3 - U1, V3 - V1, 0, 0),
-				FPlane(U1, V1, 1, 0),
-				FPlane(0, 0, 0, 1)
+				float	U1 = Wedges[Faces[FaceIndex].iWedge[0]].UVs[0].X,
+					U2 = Wedges[Faces[FaceIndex].iWedge[1]].UVs[0].X,
+					U3 = Wedges[Faces[FaceIndex].iWedge[2]].UVs[0].X,
+					V1 = Wedges[Faces[FaceIndex].iWedge[0]].UVs[0].Y,
+					V2 = Wedges[Faces[FaceIndex].iWedge[1]].UVs[0].Y,
+					V3 = Wedges[Faces[FaceIndex].iWedge[2]].UVs[0].Y;
+
+				FMatrix	ParameterToTexture(
+					FPlane(U2 - U1, V2 - V1, 0, 0),
+					FPlane(U3 - U1, V3 - V1, 0, 0),
+					FPlane(U1, V1, 1, 0),
+					FPlane(0, 0, 0, 1)
 				);
 
-			FMatrix	TextureToLocal = ParameterToTexture.Inverse() * ParameterToLocal;
-			FVector	TangentX = TextureToLocal.TransformVector(FVector(1, 0, 0)).GetSafeNormal(),
-				TangentY = TextureToLocal.TransformVector(FVector(0, 1, 0)).GetSafeNormal(),
-				TangentZ;
+				FMatrix	TextureToLocal = ParameterToTexture.Inverse() * ParameterToLocal;
+				FVector	TangentX = TextureToLocal.TransformVector(FVector(1, 0, 0)).GetSafeNormal(),
+					TangentY = TextureToLocal.TransformVector(FVector(0, 1, 0)).GetSafeNormal(),
+					TangentZ;
 
-			TangentX = TangentX - TriangleNormal * (TangentX | TriangleNormal);
-			TangentY = TangentY - TriangleNormal * (TangentY | TriangleNormal);
+				TangentX = TangentX - TriangleNormal * (TangentX | TriangleNormal);
+				TangentY = TangentY - TriangleNormal * (TangentY | TriangleNormal);
 
-			FaceTangentX[FaceIndex] = TangentX.GetSafeNormal();
-			FaceTangentY[FaceIndex] = TangentY.GetSafeNormal();
+				FaceTangentX[FaceIndex] = TangentX.GetSafeNormal();
+				FaceTangentY[FaceIndex] = TangentY.GetSafeNormal();
+			}
+			else
+			{
+				FaceTangentX[FaceIndex] = FVector::ZeroVector;
+				FaceTangentY[FaceIndex] = FVector::ZeroVector;
+			}
 		}
 	}
 
