Remove SkeletonEditor as a depends for MeshUtilites as its no longer used by MeshUtilites

#deadend

diff --git a/Engine/Source/Developer/MeshUtilities/MeshUtilities.Build.cs b/Engine/Source/Developer/MeshUtilities/MeshUtilities.Build.cs
--- a/Engine/Source/Developer/MeshUtilities/MeshUtilities.Build.cs
+++ b/Engine/Source/Developer/MeshUtilities/MeshUtilities.Build.cs
@@ -33,7 +33,6 @@ public class MeshUtilities : ModuleRules
 				"AnimationBlueprintEditor",
 				"AnimationEditor",
 				"SkeletalMeshEditor",
-				"SkeletonEditor",
 				"PropertyEditor",
 				"EditorStyle",
                 "GraphColor",
