Copying //UE4/Dev-Editor to Dev-Main (//UE4/Dev-Main) for 4.24

diff --git a/Engine/Source/Developer/MeshBuilder/Private/StaticMeshBuilder.cpp b/Engine/Source/Developer/MeshBuilder/Private/StaticMeshBuilder.cpp
--- a/Engine/Source/Developer/MeshBuilder/Private/StaticMeshBuilder.cpp
+++ b/Engine/Source/Developer/MeshBuilder/Private/StaticMeshBuilder.cpp
@@ -4,9 +4,8 @@
 #include "Engine/StaticMesh.h"
 #include "StaticMeshResources.h"
 #include "PhysicsEngine/BodySetup.h"
-#include "MeshDescription.h"
+#include "StaticMeshAttributes.h"
 #include "MeshDescriptionOperations.h"
-#include "MeshAttributes.h"
 #include "MeshDescriptionHelper.h"
 #include "BuildOptimizationHelper.h"
 #include "Components.h"
@@ -91,26 +90,26 @@ bool FStaticMeshBuilder::Build(FStaticMeshRenderData& StaticMeshRenderData, USta
 
 		float MaxDeviation = 0.0f;
 		FMeshBuildSettings& LODBuildSettings = SrcModel.BuildSettings;
-		const FMeshDescription* OriginalMeshDescription = StaticMesh->GetMeshDescription(LodIndex);
+		const FMeshDescription* MeshDescription = StaticMesh->GetMeshDescription(LodIndex);
 		FMeshDescriptionHelper MeshDescriptionHelper(&LODBuildSettings);
 
 		FMeshReductionSettings ReductionSettings = LODGroup.GetSettings(SrcModel.ReductionSettings, LodIndex);
 
 		//Make sure we do not reduce a non custom LOD by himself
-		const int32 BaseReduceLodIndex = FMath::Clamp<int32>(ReductionSettings.BaseLODModel, 0, OriginalMeshDescription == nullptr ? LodIndex - 1 : LodIndex);
+		const int32 BaseReduceLodIndex = FMath::Clamp<int32>(ReductionSettings.BaseLODModel, 0, MeshDescription == nullptr ? LodIndex - 1 : LodIndex);
 		// Use simplifier if a reduction in triangles or verts has been requested.
 		bool bUseReduction = StaticMesh->IsReductionActive(LodIndex);
 
-		if (OriginalMeshDescription != nullptr)
+		if (MeshDescription != nullptr)
 		{
-			MeshDescriptionHelper.GetRenderMeshDescription(StaticMesh, *OriginalMeshDescription, MeshDescriptions[LodIndex]);
+			MeshDescriptionHelper.GetRenderMeshDescription(StaticMesh, *MeshDescription, MeshDescriptions[LodIndex]);
 		}
 		else
 		{
 			if (bUseReduction)
 			{
 				// Initialize an empty mesh description that the reduce will fill
-				UStaticMesh::RegisterMeshAttributes(MeshDescriptions[LodIndex]);
+				FStaticMeshAttributes(MeshDescriptions[LodIndex]).Register();
 			}
 			else
 			{
diff --git a/Engine/Source/Developer/MeshBuilder/Public/IMeshBuilderModule.h b/Engine/Source/Developer/MeshBuilder/Public/IMeshBuilderModule.h
--- a/Engine/Source/Developer/MeshBuilder/Public/IMeshBuilderModule.h
+++ b/Engine/Source/Developer/MeshBuilder/Public/IMeshBuilderModule.h
@@ -32,4 +32,5 @@ public:
 	}
 
 	virtual bool BuildMesh(class FStaticMeshRenderData& OutRenderData, class UObject* Mesh, const class FStaticMeshLODGroup& LODGroup) { return false; }
+	virtual bool BuildSkeletalMesh(class USkeletalMesh* SkeletalMesh, int32 LODIndex, const bool bRegenDepLODs) { return false; }
 };
diff --git a/Engine/Source/Developer/MeshBuilder/Public/MeshBuilder.h b/Engine/Source/Developer/MeshBuilder/Public/MeshBuilder.h
--- a/Engine/Source/Developer/MeshBuilder/Public/MeshBuilder.h
+++ b/Engine/Source/Developer/MeshBuilder/Public/MeshBuilder.h
@@ -4,6 +4,11 @@
 
 #include "CoreMinimal.h"
 
+class UStaticMesh;
+class FStaticMeshRenderData;
+class FStaticMeshLODGroup;
+class USkeletalMesh;
+
 /**
  * Abstract class which is the base class of all builder.
  * All share code to build some render data should be found inside this class
@@ -14,9 +19,14 @@ public:
 	FMeshBuilder();
 	
 	/**
-	 * Build function should be override and is the starting point for all mesh builders
+	 * Build function should be override and is the starting point for static mesh builders
+	 */
+	virtual bool Build(FStaticMeshRenderData& OutRenderData, UStaticMesh* StaticMesh, const FStaticMeshLODGroup& LODGroup) = 0;
+
+	/**
+	 * Build function should be override and is the starting point for skeletal mesh builders
 	 */
-	virtual bool Build(class FStaticMeshRenderData& OutRenderData, class UStaticMesh* StaticMesh, const class FStaticMeshLODGroup& LODGroup) =0;
+	virtual bool Build(USkeletalMesh* SkeletalMesh, const int32 LODIndex, const bool bRegenDepLODs) = 0;
 
 private:
 
diff --git a/Engine/Source/Developer/MeshBuilder/Public/SkeletalMeshBuilder.h b/Engine/Source/Developer/MeshBuilder/Public/SkeletalMeshBuilder.h
new file mode 100644
--- /dev/null
+++ b/Engine/Source/Developer/MeshBuilder/Public/SkeletalMeshBuilder.h
@@ -0,0 +1,35 @@
+// Copyright 1998-2019 Epic Games, Inc. All Rights Reserved.
+
+#pragma once
+
+#include "CoreMinimal.h"
+#include "MeshBuilder.h"
+
+DECLARE_LOG_CATEGORY_EXTERN(LogSkeletalMeshBuilder, Log, All);
+
+class UStaticMesh;
+class FStaticMeshRenderData;
+class FStaticMeshLODGroup;
+class USkeletalMesh;
+
+class MESHBUILDER_API FSkeletalMeshBuilder : public FMeshBuilder
+{
+public:
+	FSkeletalMeshBuilder();
+
+	//No support for static mesh build in this class
+	virtual bool Build(FStaticMeshRenderData& OutRenderData, UStaticMesh* StaticMesh, const FStaticMeshLODGroup& LODGroup) override
+	{
+		bool No_Support_For_StaticMesh_Build_In_FSkeletalMeshBuilder_Class = false;
+		check(No_Support_For_StaticMesh_Build_In_FSkeletalMeshBuilder_Class);
+		return false;
+	}
+	
+	virtual bool Build(USkeletalMesh* SkeletalMesh, const int32 LODIndex, const bool bRegenDepLODs) override;
+
+	virtual ~FSkeletalMeshBuilder() {}
+
+private:
+
+};
+
diff --git a/Engine/Source/Developer/MeshBuilder/Public/StaticMeshBuilder.h b/Engine/Source/Developer/MeshBuilder/Public/StaticMeshBuilder.h
--- a/Engine/Source/Developer/MeshBuilder/Public/StaticMeshBuilder.h
+++ b/Engine/Source/Developer/MeshBuilder/Public/StaticMeshBuilder.h
@@ -7,12 +7,26 @@
 
 DECLARE_LOG_CATEGORY_EXTERN(LogStaticMeshBuilder, Log, All);
 
+class UStaticMesh;
+class FStaticMeshRenderData;
+class FStaticMeshLODGroup;
+class USkeletalMesh;
+
 class MESHBUILDER_API FStaticMeshBuilder : public FMeshBuilder
 {
 public:
 	FStaticMeshBuilder();
 
-	virtual bool Build(class FStaticMeshRenderData& OutRenderData, class UStaticMesh* StaticMesh, const class FStaticMeshLODGroup& LODGroup) override;
+	virtual bool Build(FStaticMeshRenderData& OutRenderData, UStaticMesh* StaticMesh, const FStaticMeshLODGroup& LODGroup) override;
+
+	//No support for skeletal mesh build in this class
+	virtual bool Build(USkeletalMesh* SkeletalMesh, const int32 LODIndex, const bool bRegenDepLODs) override
+	{
+		bool No_Support_For_SkeletalMesh_Build_In_FStaticMeshBuilder_Class = false;
+		check(No_Support_For_SkeletalMesh_Build_In_FStaticMeshBuilder_Class);
+		return false;
+	}
+
 	virtual ~FStaticMeshBuilder() {}
 
 private:
diff --git a/Engine/Source/Developer/MeshDescriptionOperations/MeshDescriptionOperations.Build.cs b/Engine/Source/Developer/MeshDescriptionOperations/MeshDescriptionOperations.Build.cs
--- a/Engine/Source/Developer/MeshDescriptionOperations/MeshDescriptionOperations.Build.cs
+++ b/Engine/Source/Developer/MeshDescriptionOperations/MeshDescriptionOperations.Build.cs
@@ -15,7 +15,8 @@ namespace UnrealBuildTool.Rules
                     "Core",
                     "CoreUObject",
                     "RenderCore",
-                    "MeshDescription",
+					"MeshDescription",
+                    "StaticMeshDescription",
                     "MeshUtilitiesCommon",
                     "RawMesh",
                     "RHI"
diff --git a/Engine/Source/Developer/MeshDescriptionOperations/Private/MeshDescriptionOperations.cpp b/Engine/Source/Developer/MeshDescriptionOperations/Private/MeshDescriptionOperations.cpp
--- a/Engine/Source/Developer/MeshDescriptionOperations/Private/MeshDescriptionOperations.cpp
+++ b/Engine/Source/Developer/MeshDescriptionOperations/Private/MeshDescriptionOperations.cpp
@@ -2,8 +2,7 @@
 
 #include "MeshDescriptionOperations.h"
 #include "UObject/Package.h"
-#include "MeshDescription.h"
-#include "MeshAttributes.h"
+#include "StaticMeshAttributes.h"
 #include "RawMesh.h"
 #include "LayoutUV.h"
 #include "OverlappingCorners.h"
diff --git a/Engine/Source/Developer/MeshMergeUtilities/Private/MeshMergeDataTracker.cpp b/Engine/Source/Developer/MeshMergeUtilities/Private/MeshMergeDataTracker.cpp
--- a/Engine/Source/Developer/MeshMergeUtilities/Private/MeshMergeDataTracker.cpp
+++ b/Engine/Source/Developer/MeshMergeUtilities/Private/MeshMergeDataTracker.cpp
@@ -4,6 +4,7 @@
 #include "MeshMergeHelpers.h"
 #include "Misc/Crc.h"
 #include "Engine/StaticMesh.h"
+#include "StaticMeshAttributes.h"
 
 FMeshMergeDataTracker::FMeshMergeDataTracker()
 	: AvailableLightMapUVChannel(INDEX_NONE), SummedLightMapPixels(0)
@@ -16,7 +17,7 @@ FMeshDescription& FMeshMergeDataTracker::AddAndRetrieveRawMesh(int32 MeshIndex,
 {
 	checkf(!RawMeshLODs.Contains(FMeshLODKey(MeshIndex, LODIndex, InMesh)), TEXT("Raw Mesh already added for this key"));
 	FMeshDescription& MeshDescription = RawMeshLODs.Add(FMeshLODKey(MeshIndex, LODIndex, InMesh));
-	UStaticMesh::RegisterMeshAttributes(MeshDescription);
+	FStaticMeshAttributes(MeshDescription).Register();
 	return MeshDescription;
 }
 
diff --git a/Engine/Source/Developer/MeshMergeUtilities/Private/MeshMergeHelpers.cpp b/Engine/Source/Developer/MeshMergeUtilities/Private/MeshMergeHelpers.cpp
--- a/Engine/Source/Developer/MeshMergeUtilities/Private/MeshMergeHelpers.cpp
+++ b/Engine/Source/Developer/MeshMergeUtilities/Private/MeshMergeHelpers.cpp
@@ -6,9 +6,7 @@
 #include "Engine/MeshMerging.h"
 
 #include "MaterialOptions.h"
-#include "MeshDescription.h"
-#include "MeshAttributes.h"
-#include "MeshAttributeArray.h"
+#include "StaticMeshAttributes.h"
 #include "MeshDescriptionOperations.h"
 
 #include "Misc/PackageName.h"
@@ -169,7 +167,7 @@ void FMeshMergeHelpers::ExtractSections(const UStaticMesh* StaticMesh, int32 LOD
 void FMeshMergeHelpers::ExpandInstances(const UInstancedStaticMeshComponent* InInstancedStaticMeshComponent, FMeshDescription& InOutRawMesh, TArray<FSectionInfo>& InOutSections)
 {
 	FMeshDescription CombinedRawMesh;
-	UStaticMesh::RegisterMeshAttributes(CombinedRawMesh);
+	FStaticMeshAttributes(CombinedRawMesh).Register();
 
 	FTransform ComponentTransform = InInstancedStaticMeshComponent->GetComponentTransform();
 	FTransform ComponentTransformInv = ComponentTransform.Inverse();
@@ -864,7 +862,7 @@ void FMeshMergeHelpers::RetrieveCullingLandscapeAndVolumes(UWorld* InWorld, cons
 	{
 		// Export the landscape to raw mesh format
 		FMeshDescription* MeshDescription = new FMeshDescription();
-		UStaticMesh::RegisterMeshAttributes(*MeshDescription);
+		FStaticMeshAttributes(*MeshDescription).Register();
 		FBoxSphereBounds LandscapeBounds = EstimatedMeshProxyBounds;
 		Landscape->ExportToRawMesh(LandscapeExportLOD, *MeshDescription, LandscapeBounds);
 		if (MeshDescription->Vertices().Num())
@@ -878,14 +876,15 @@ void FMeshMergeHelpers::RetrieveCullingLandscapeAndVolumes(UWorld* InWorld, cons
 	{
 		// Export the landscape to raw mesh format
 		FMeshDescription* VolumeMesh = new FMeshDescription();
-		UStaticMesh::RegisterMeshAttributes(*VolumeMesh);
+		FStaticMeshAttributes MeshAttributes(*VolumeMesh);
+		MeshAttributes.Register();
 
 		TArray<FStaticMaterial>	VolumeMaterials;
 		GetBrushMesh(Volume, Volume->Brush, *VolumeMesh, VolumeMaterials);
 
 		// Offset vertices to correct world position;
 		FVector VolumeLocation = Volume->GetActorLocation();
-		TVertexAttributesRef<FVector> VertexPositions = VolumeMesh->VertexAttributes().GetAttributesRef<FVector>(MeshAttribute::Vertex::Position);
+		TVertexAttributesRef<FVector> VertexPositions = MeshAttributes.GetVertexPositions();
 		for(const FVertexID& VertexID : VolumeMesh->Vertices().GetElementIDs())
 		{
 			VertexPositions[VertexID] += VolumeLocation;
@@ -1289,7 +1288,7 @@ void FMeshMergeHelpers::ExtractImposterToRawMesh(const UStaticMeshComponent* InI
 {
 	// Retrieve imposter LOD mesh	
 	const int32 LODIndex = InImposterComponent->GetStaticMesh()->GetNumLODs() - 1;
-	UStaticMesh::RegisterMeshAttributes(InRawMesh);
+	FStaticMeshAttributes(InRawMesh).Register();
 	FMeshMergeHelpers::RetrieveMesh(InImposterComponent->GetStaticMesh(), LODIndex, InRawMesh);
 }
 
@@ -1329,7 +1328,8 @@ void FMeshMergeHelpers::MergeImpostersToRawMesh(TArray<const UStaticMeshComponen
 
 			// Retrieve mesh data in FMeshDescription form
 			FMeshDescription ImposterMesh;
-			UStaticMesh::RegisterMeshAttributes(ImposterMesh);
+			FStaticMeshAttributes ImposterMeshAttributes(ImposterMesh);
+			ImposterMeshAttributes.Register();
 			FMeshMergeHelpers::RetrieveMesh(Component, LODIndex, ImposterMesh, false);
 
 			// Retrieve the sections, we're expect 1 for imposter meshes
@@ -1344,7 +1344,7 @@ void FMeshMergeHelpers::MergeImpostersToRawMesh(TArray<const UStaticMeshComponen
 
 			// Imposter magic, we're storing the actor world position and X scale spread across two UV channels
 			const int32 UVTwoIndex = UVOneIndex + 1;
-			TVertexInstanceAttributesRef<FVector2D> VertexInstanceUVs = ImposterMesh.VertexInstanceAttributes().GetAttributesRef<FVector2D>(MeshAttribute::VertexInstance::TextureCoordinate);
+			TVertexInstanceAttributesRef<FVector2D> VertexInstanceUVs = ImposterMeshAttributes.GetVertexInstanceUVs();
 			VertexInstanceUVs.SetNumIndices(UVTwoIndex + 1);
 			const int32 NumIndices = ImposterMesh.VertexInstances().Num();
 			const FTransform& ActorToWorld = Component->GetOwner()->GetActorTransform();
@@ -1363,7 +1363,7 @@ void FMeshMergeHelpers::MergeImpostersToRawMesh(TArray<const UStaticMeshComponen
 				VertexInstanceUVs.Set(VertexInstanceID, UVTwoIndex, UVTwo);
 			}
 
-			TPolygonGroupAttributesRef<FName> SourcePolygonGroupImportedMaterialSlotNames = ImposterMesh.PolygonGroupAttributes().GetAttributesRef<FName>(MeshAttribute::PolygonGroup::ImportedMaterialSlotName);
+			TPolygonGroupAttributesRef<FName> SourcePolygonGroupImportedMaterialSlotNames = ImposterMeshAttributes.GetPolygonGroupMaterialSlotNames();
 			TPolygonGroupAttributesRef<FName> TargetPolygonGroupImportedMaterialSlotNames = InRawMesh.PolygonGroupAttributes().GetAttributesRef<FName>(MeshAttribute::PolygonGroup::ImportedMaterialSlotName);
 
 			//Add the missing polygon group ID to the target(InRawMesh)
diff --git a/Engine/Source/Developer/MeshMergeUtilities/Private/MeshMergeUtilities.cpp b/Engine/Source/Developer/MeshMergeUtilities/Private/MeshMergeUtilities.cpp
--- a/Engine/Source/Developer/MeshMergeUtilities/Private/MeshMergeUtilities.cpp
+++ b/Engine/Source/Developer/MeshMergeUtilities/Private/MeshMergeUtilities.cpp
@@ -68,9 +68,7 @@
 #include "IMeshMergeExtension.h"
 
 #include "RawMesh.h"
-#include "MeshDescription.h"
-#include "MeshAttributes.h"
-#include "MeshAttributeArray.h"
+#include "StaticMeshAttributes.h"
 #include "MeshDescriptionOperations.h"
 
 #if WITH_EDITOR
@@ -139,7 +137,7 @@ void FMeshMergeUtilities::BakeMaterialsForComponent(TArray<TWeakObjectPtr<UObjec
 		if (bProcessedLOD)
 		{
 			FMeshDescription& RawMesh = RawMeshLODs.Add(LODIndex);
-			UStaticMesh::RegisterMeshAttributes(RawMesh);
+			FStaticMeshAttributes(RawMesh).Register();
 			Adapter->RetrieveRawMeshData(LODIndex, RawMesh, MaterialOptions->bUseMeshData);
 		}
 
@@ -1194,7 +1192,7 @@ void FMeshMergeUtilities::CreateProxyMesh(const TArray<UStaticMeshComponent*>& I
 			// Retrieve mesh data in FMeshDescription form
 			const int32 MeshIndex = MeshDescriptionData.Add(new FMeshDescription());
 			FMeshDescription& MeshDescription = *MeshDescriptionData[MeshIndex];
-			UStaticMesh::RegisterMeshAttributes(MeshDescription);
+			FStaticMeshAttributes(MeshDescription).Register();
 			FMeshMergeHelpers::RetrieveMesh(StaticMeshComponent, LODIndex, MeshDescription, bPropagateVertexColours);
 
 			// Reset section array for reuse
@@ -2510,8 +2508,7 @@ void FMeshMergeUtilities::MergeComponentsToStaticMesh(const TArray<UPrimitiveCom
 					}
 				}
 
-				FMeshDescription* MeshDescription = StaticMesh->CreateMeshDescription(LODIndex);
-				*MeshDescription = MergedMeshLOD;
+				FMeshDescription* MeshDescription = StaticMesh->CreateMeshDescription(LODIndex, MergedMeshLOD);
 				StaticMesh->CommitMeshDescription(LODIndex);
 			}
 		}
@@ -2645,7 +2642,7 @@ void FMeshMergeUtilities::CreateMergedRawMeshes(FMeshMergeDataTracker& InDataTra
 			// Find meshes for each lod
 			const int32 LODIndex = *Iterator;
 			FMeshDescription& MergedMesh = OutMergedRawMeshes[LODIndex];
-			UStaticMesh::RegisterMeshAttributes(MergedMesh);
+			FStaticMeshAttributes(MergedMesh).Register();
 
 			for (int32 ComponentIndex = 0; ComponentIndex < InStaticMeshComponentsToMerge.Num(); ++ComponentIndex)
 			{
@@ -2744,9 +2741,8 @@ void FMeshMergeUtilities::CreateMergedRawMeshes(FMeshMergeDataTracker& InDataTra
 	}
 	else
 	{	
-		OutMergedRawMeshes.AddZeroed(1);
-		FMeshDescription& MergedMesh = OutMergedRawMeshes.Last();
-		UStaticMesh::RegisterMeshAttributes(MergedMesh);
+		FMeshDescription& MergedMesh = OutMergedRawMeshes.AddDefaulted_GetRef();
+		FStaticMeshAttributes(MergedMesh).Register();
 
 		for (int32 ComponentIndex = 0; ComponentIndex < InStaticMeshComponentsToMerge.Num(); ++ComponentIndex)
 		{
diff --git a/Engine/Source/Developer/MeshMergeUtilities/Private/ProxyGenerationProcessor.cpp b/Engine/Source/Developer/MeshMergeUtilities/Private/ProxyGenerationProcessor.cpp
--- a/Engine/Source/Developer/MeshMergeUtilities/Private/ProxyGenerationProcessor.cpp
+++ b/Engine/Source/Developer/MeshMergeUtilities/Private/ProxyGenerationProcessor.cpp
@@ -8,10 +8,7 @@
 #include "IMeshReductionInterfaces.h"
 #include "IMeshReductionManagerModule.h"
 #include "Modules/ModuleManager.h"
-
-#include "MeshDescription.h"
-#include "MeshAttributes.h"
-#include "MeshAttributeArray.h"
+#include "StaticMeshAttributes.h"
 
 #if WITH_EDITOR
 #include "Editor.h"
@@ -226,11 +223,9 @@ void FProxyGenerationProcessor::ProcessJob(const FGuid& JobGuid, FProxyGeneratio
 
 		//Commit the FMeshDescription to the source model we just created
 		int32 SourceModelIndex = StaticMesh->GetNumSourceModels() - 1;
-		FMeshDescription* MeshDescription = StaticMesh->CreateMeshDescription(SourceModelIndex);
+		FMeshDescription* MeshDescription = StaticMesh->CreateMeshDescription(SourceModelIndex, Data->RawMesh);
 		if (ensure(MeshDescription))
 		{
-			*MeshDescription = Data->RawMesh;
-
 			// Make sure the Proxy material have a valid ImportedMaterialSlotName
 			//The proxy material must be add only once and is always the first slot of the HLOD mesh
 			FStaticMaterial NewMaterial(ProxyMaterial);
@@ -238,7 +233,7 @@ void FProxyGenerationProcessor::ProcessJob(const FGuid& JobGuid, FProxyGeneratio
 			{
 				NewMaterial.ImportedMaterialSlotName = PolygonGroupMaterialSlotName[MeshDescription->PolygonGroups().GetFirstValidID()];
 			}
-				StaticMesh->StaticMaterials.Add(NewMaterial);
+			StaticMesh->StaticMaterials.Add(NewMaterial);
 
 			StaticMesh->CommitMeshDescription(SourceModelIndex);
 		}
diff --git a/Engine/Source/Developer/MeshReductionInterface/Public/IMeshReductionInterfaces.h b/Engine/Source/Developer/MeshReductionInterface/Public/IMeshReductionInterfaces.h
--- a/Engine/Source/Developer/MeshReductionInterface/Public/IMeshReductionInterfaces.h
+++ b/Engine/Source/Developer/MeshReductionInterface/Public/IMeshReductionInterfaces.h
@@ -39,8 +39,7 @@ public:
 	*/
 	virtual bool ReduceSkeletalMesh(
 		class USkeletalMesh* SkeletalMesh,
-		int32 LODIndex,
-		bool bReregisterComponent = true
+		int32 LODIndex
 	) = 0;
 	/**
 	* Returns a unique string identifying both the reduction plugin itself and the version of the plugin.
diff --git a/Engine/Source/Developer/MeshSimplifier/QuadricMeshReduction.Build.cs b/Engine/Source/Developer/MeshSimplifier/QuadricMeshReduction.Build.cs
--- a/Engine/Source/Developer/MeshSimplifier/QuadricMeshReduction.Build.cs
+++ b/Engine/Source/Developer/MeshSimplifier/QuadricMeshReduction.Build.cs
@@ -14,6 +14,7 @@ public class QuadricMeshReduction : ModuleRules
         PrivateDependencyModuleNames.Add("MeshDescription");
         PrivateDependencyModuleNames.Add("MeshDescriptionOperations");
         PrivateDependencyModuleNames.Add("MeshUtilitiesCommon");
+        PrivateDependencyModuleNames.Add("StaticMeshDescription");
 
         PrivateIncludePathModuleNames.AddRange(
         new string[] {
diff --git a/Engine/Source/Developer/MeshUtilities/Private/MeshUtilities.cpp b/Engine/Source/Developer/MeshUtilities/Private/MeshUtilities.cpp
--- a/Engine/Source/Developer/MeshUtilities/Private/MeshUtilities.cpp
+++ b/Engine/Source/Developer/MeshUtilities/Private/MeshUtilities.cpp
@@ -94,7 +94,7 @@
 #include "OverlappingCorners.h"
 #include "MeshUtilitiesCommon.h"
 
-#include "MeshDescription.h"
+#include "StaticMeshAttributes.h"
 #include "MeshDescriptionOperations.h"
 
 #if WITH_EDITOR
@@ -645,6 +645,9 @@ void FMeshUtilities::BuildSkeletalModelFromChunks(FSkeletalMeshLODModel& LODMode
 		Section.MaterialIndex = SrcChunk->MaterialIndex;
 		Exchange(Section.BoneMap, SrcChunk->BoneMap);
 
+		Section.OriginalDataSectionIndex = SrcChunk->OriginalSectionIndex;
+		Section.ChunkedParentSectionIndex = SrcChunk->ParentChunkSectionIndex;
+
 		// Update the active bone indices on the LOD model.
 		for (int32 BoneIndex = 0; BoneIndex < Section.BoneMap.Num(); ++BoneIndex)
 		{
@@ -2629,10 +2632,13 @@ public:
 				FRawMesh& DestMesh = LODMeshes[NumValidLODs];
 				FOverlappingCorners& InOverlappingCorners = LODOverlappingCorners[ReductionSettings.BaseLODModel];
 				FOverlappingCorners& DestOverlappingCorners = LODOverlappingCorners[NumValidLODs];
+
 				FMeshDescription SrcMeshdescription;
-				UStaticMesh::RegisterMeshAttributes(SrcMeshdescription);
+				FStaticMeshAttributes(SrcMeshdescription).Register();
+
 				FMeshDescription DestMeshdescription;
-				UStaticMesh::RegisterMeshAttributes(DestMeshdescription);
+				FStaticMeshAttributes(DestMeshdescription).Register();
+
 				TMap<int32, FName> FromMaterialMap;
 				FMeshDescriptionOperations::ConvertFromRawMesh(InMesh, SrcMeshdescription, FromMaterialMap);
 				MeshReduction->ReduceMeshDescription(DestMeshdescription, LODMaxDeviation[NumValidLODs], SrcMeshdescription, InOverlappingCorners, ReductionSettings);
@@ -3450,342 +3456,6 @@ public:
 		}
 	}
 
-	void Skeletal_ComputeTangents(
-		IMeshBuildData* BuildData,
-		const FOverlappingCorners& OverlappingCorners
-		)
-	{
-		bool bBlendOverlappingNormals = true;
-		bool bIgnoreDegenerateTriangles = BuildData->BuildOptions.bRemoveDegenerateTriangles;
-		bool bComputeWeightedNormals = BuildData->BuildOptions.bComputeWeightedNormals;
-		
-		// Compute per-triangle tangents.
-		TArray<FVector> TriangleTangentX;
-		TArray<FVector> TriangleTangentY;
-		TArray<FVector> TriangleTangentZ;
-
-		Skeletal_ComputeTriangleTangents(
-			TriangleTangentX,
-			TriangleTangentY,
-			TriangleTangentZ,
-			BuildData,
-			bIgnoreDegenerateTriangles ? SMALL_NUMBER : FLT_MIN
-			);
-
-		TArray<FVector>& WedgeTangentX = BuildData->GetTangentArray(0);
-		TArray<FVector>& WedgeTangentY = BuildData->GetTangentArray(1);
-		TArray<FVector>& WedgeTangentZ = BuildData->GetTangentArray(2);
-
-		// Declare these out here to avoid reallocations.
-		TArray<FFanFace> RelevantFacesForCorner[3];
-		TArray<int32> AdjacentFaces;
-
-		int32 NumFaces = BuildData->GetNumFaces();
-		int32 NumWedges = BuildData->GetNumWedges();
-		check(NumFaces * 3 <= NumWedges);
-
-		// Allocate storage for tangents if none were provided.
-		if (WedgeTangentX.Num() != NumWedges)
-		{
-			WedgeTangentX.Empty(NumWedges);
-			WedgeTangentX.AddZeroed(NumWedges);
-		}
-		if (WedgeTangentY.Num() != NumWedges)
-		{
-			WedgeTangentY.Empty(NumWedges);
-			WedgeTangentY.AddZeroed(NumWedges);
-		}
-		if (WedgeTangentZ.Num() != NumWedges)
-		{
-			WedgeTangentZ.Empty(NumWedges);
-			WedgeTangentZ.AddZeroed(NumWedges);
-		}
-
-		for (int32 FaceIndex = 0; FaceIndex < NumFaces; FaceIndex++)
-		{
-			int32 WedgeOffset = FaceIndex * 3;
-			FVector CornerPositions[3];
-			FVector CornerTangentX[3];
-			FVector CornerTangentY[3];
-			FVector CornerTangentZ[3];
-
-			for (int32 CornerIndex = 0; CornerIndex < 3; CornerIndex++)
-			{
-				CornerTangentX[CornerIndex] = FVector::ZeroVector;
-				CornerTangentY[CornerIndex] = FVector::ZeroVector;
-				CornerTangentZ[CornerIndex] = FVector::ZeroVector;
-				CornerPositions[CornerIndex] = BuildData->GetVertexPosition(FaceIndex, CornerIndex);
-				RelevantFacesForCorner[CornerIndex].Reset();
-			}
-
-			// Don't process degenerate triangles.
-			if (PointsEqual(CornerPositions[0], CornerPositions[1], BuildData->BuildOptions.OverlappingThresholds)
-				|| PointsEqual(CornerPositions[0], CornerPositions[2], BuildData->BuildOptions.OverlappingThresholds)
-				|| PointsEqual(CornerPositions[1], CornerPositions[2], BuildData->BuildOptions.OverlappingThresholds))
-			{
-				continue;
-			}
-
-			// No need to process triangles if tangents already exist.
-			bool bCornerHasTangents[3] = { 0 };
-			for (int32 CornerIndex = 0; CornerIndex < 3; CornerIndex++)
-			{
-				bCornerHasTangents[CornerIndex] = !WedgeTangentX[WedgeOffset + CornerIndex].IsZero()
-					&& !WedgeTangentY[WedgeOffset + CornerIndex].IsZero()
-					&& !WedgeTangentZ[WedgeOffset + CornerIndex].IsZero();
-			}
-			if (bCornerHasTangents[0] && bCornerHasTangents[1] && bCornerHasTangents[2])
-			{
-				continue;
-			}
-
-			// Calculate smooth vertex normals.
-			float Determinant = FVector::Triple(
-				TriangleTangentX[FaceIndex],
-				TriangleTangentY[FaceIndex],
-				TriangleTangentZ[FaceIndex]
-				);
-
-			// Start building a list of faces adjacent to this face.
-			AdjacentFaces.Reset();
-			for (int32 CornerIndex = 0; CornerIndex < 3; CornerIndex++)
-			{
-				int32 ThisCornerIndex = WedgeOffset + CornerIndex;
-				const TArray<int32>& DupVerts = OverlappingCorners.FindIfOverlapping(ThisCornerIndex);
-				if (DupVerts.Num() == 0)
-				{
-					AdjacentFaces.AddUnique(ThisCornerIndex / 3); // I am a "dup" of myself
-				}
-				for (int32 k = 0; k < DupVerts.Num(); k++)
-				{
-					AdjacentFaces.AddUnique(DupVerts[k] / 3);
-				}
-			}
-
-			// We need to sort these here because the criteria for point equality is
-			// exact, so we must ensure the exact same order for all dups.
-			AdjacentFaces.Sort();
-
-			// Process adjacent faces
-			for (int32 AdjacentFaceIndex = 0; AdjacentFaceIndex < AdjacentFaces.Num(); AdjacentFaceIndex++)
-			{
-				int32 OtherFaceIndex = AdjacentFaces[AdjacentFaceIndex];
-				for (int32 OurCornerIndex = 0; OurCornerIndex < 3; OurCornerIndex++)
-				{
-					if (bCornerHasTangents[OurCornerIndex])
-						continue;
-
-					FFanFace NewFanFace;
-					int32 CommonIndexCount = 0;
-
-					// Check for vertices in common.
-					if (FaceIndex == OtherFaceIndex)
-					{
-						CommonIndexCount = 3;
-						NewFanFace.LinkedVertexIndex = OurCornerIndex;
-					}
-					else
-					{
-						// Check matching vertices against main vertex .
-						for (int32 OtherCornerIndex = 0; OtherCornerIndex < 3; OtherCornerIndex++)
-						{
-							if (PointsEqual(
-								CornerPositions[OurCornerIndex],
-								BuildData->GetVertexPosition(OtherFaceIndex, OtherCornerIndex),
-								BuildData->BuildOptions.OverlappingThresholds
-								))
-							{
-								CommonIndexCount++;
-								NewFanFace.LinkedVertexIndex = OtherCornerIndex;
-							}
-						}
-					}
-
-					// Add if connected by at least one point. Smoothing matches are considered later.
-					if (CommonIndexCount > 0)
-					{
-						NewFanFace.FaceIndex = OtherFaceIndex;
-						NewFanFace.bFilled = (OtherFaceIndex == FaceIndex); // Starter face for smoothing floodfill.
-						NewFanFace.bBlendTangents = NewFanFace.bFilled;
-						NewFanFace.bBlendNormals = NewFanFace.bFilled;
-						RelevantFacesForCorner[OurCornerIndex].Add(NewFanFace);
-					}
-				}
-			}
-
-			// Find true relevance of faces for a vertex normal by traversing
-			// smoothing-group-compatible connected triangle fans around common vertices.
-			for (int32 CornerIndex = 0; CornerIndex < 3; CornerIndex++)
-			{
-				if (bCornerHasTangents[CornerIndex])
-					continue;
-
-				int32 NewConnections;
-				do
-				{
-					NewConnections = 0;
-					for (int32 OtherFaceIdx = 0; OtherFaceIdx < RelevantFacesForCorner[CornerIndex].Num(); OtherFaceIdx++)
-					{
-						FFanFace& OtherFace = RelevantFacesForCorner[CornerIndex][OtherFaceIdx];
-						// The vertex' own face is initially the only face with bFilled == true.
-						if (OtherFace.bFilled)
-						{
-							for (int32 NextFaceIndex = 0; NextFaceIndex < RelevantFacesForCorner[CornerIndex].Num(); NextFaceIndex++)
-							{
-								FFanFace& NextFace = RelevantFacesForCorner[CornerIndex][NextFaceIndex];
-								if (!NextFace.bFilled) // && !NextFace.bBlendTangents)
-								{
-									if ((NextFaceIndex != OtherFaceIdx)
-										&& (BuildData->GetFaceSmoothingGroups(NextFace.FaceIndex) & BuildData->GetFaceSmoothingGroups(OtherFace.FaceIndex)))
-									{
-										int32 CommonVertices = 0;
-										int32 CommonTangentVertices = 0;
-										int32 CommonNormalVertices = 0;
-										for (int32 OtherCornerIndex = 0; OtherCornerIndex < 3; OtherCornerIndex++)
-										{
-											for (int32 NextCornerIndex = 0; NextCornerIndex < 3; NextCornerIndex++)
-											{
-												int32 NextVertexIndex = BuildData->GetVertexIndex(NextFace.FaceIndex, NextCornerIndex);
-												int32 OtherVertexIndex = BuildData->GetVertexIndex(OtherFace.FaceIndex, OtherCornerIndex);
-												if (PointsEqual(
-													BuildData->GetVertexPosition(NextFace.FaceIndex, NextCornerIndex),
-													BuildData->GetVertexPosition(OtherFace.FaceIndex, OtherCornerIndex),
-													BuildData->BuildOptions.OverlappingThresholds))
-												{
-													CommonVertices++;
-
-
-													if (UVsEqual(
-														BuildData->GetVertexUV(NextFace.FaceIndex, NextCornerIndex, 0),
-														BuildData->GetVertexUV(OtherFace.FaceIndex, OtherCornerIndex, 0),
-														BuildData->BuildOptions.OverlappingThresholds))
-													{
-														CommonTangentVertices++;
-													}
-													if (bBlendOverlappingNormals
-														|| NextVertexIndex == OtherVertexIndex)
-													{
-														CommonNormalVertices++;
-													}
-												}
-											}
-										}
-										// Flood fill faces with more than one common vertices which must be touching edges.
-										if (CommonVertices > 1)
-										{
-											NextFace.bFilled = true;
-											NextFace.bBlendNormals = (CommonNormalVertices > 1);
-											NewConnections++;
-
-											// Only blend tangents if there is no UV seam along the edge with this face.
-											if (OtherFace.bBlendTangents && CommonTangentVertices > 1)
-											{
-												float OtherDeterminant = FVector::Triple(
-													TriangleTangentX[NextFace.FaceIndex],
-													TriangleTangentY[NextFace.FaceIndex],
-													TriangleTangentZ[NextFace.FaceIndex]
-													);
-												if ((Determinant * OtherDeterminant) > 0.0f)
-												{
-													NextFace.bBlendTangents = true;
-												}
-											}
-										}
-									}
-								}
-							}
-						}
-					}
-				}
-				while (NewConnections > 0);
-			}
-
-			// Vertex normal construction.
-			for (int32 CornerIndex = 0; CornerIndex < 3; CornerIndex++)
-			{
-				if (bCornerHasTangents[CornerIndex])
-				{
-					CornerTangentX[CornerIndex] = WedgeTangentX[WedgeOffset + CornerIndex];
-					CornerTangentY[CornerIndex] = WedgeTangentY[WedgeOffset + CornerIndex];
-					CornerTangentZ[CornerIndex] = WedgeTangentZ[WedgeOffset + CornerIndex];
-				}
-				else
-				{
-					for (int32 RelevantFaceIdx = 0; RelevantFaceIdx < RelevantFacesForCorner[CornerIndex].Num(); RelevantFaceIdx++)
-					{
-						FFanFace const& RelevantFace = RelevantFacesForCorner[CornerIndex][RelevantFaceIdx];
-						if (RelevantFace.bFilled)
-						{
-							int32 OtherFaceIndex = RelevantFace.FaceIndex;
-							float CornerWeight = 1.0f;
-							if (bComputeWeightedNormals)
-							{
-
-								FVector OtherFacePoint[3] = { BuildData->GetVertexPosition(OtherFaceIndex, 0), BuildData->GetVertexPosition(OtherFaceIndex, 1), BuildData->GetVertexPosition(OtherFaceIndex, 2) };
-								float OtherFaceArea = TriangleUtilities::ComputeTriangleArea(OtherFacePoint[0], OtherFacePoint[1], OtherFacePoint[2]);
-								int32 OtherFaceCornerIndex = RelevantFace.LinkedVertexIndex;
-								float OtherFaceAngle = TriangleUtilities::ComputeTriangleCornerAngle(OtherFacePoint[OtherFaceCornerIndex], OtherFacePoint[(OtherFaceCornerIndex + 1) % 3], OtherFacePoint[(OtherFaceCornerIndex + 2) % 3]);
-								//Get the CornerWeight
-								CornerWeight = OtherFaceArea * OtherFaceAngle;
-							}
-
-							if (RelevantFace.bBlendTangents)
-							{
-								CornerTangentX[CornerIndex] += CornerWeight * TriangleTangentX[OtherFaceIndex];
-								CornerTangentY[CornerIndex] += CornerWeight * TriangleTangentY[OtherFaceIndex];
-							}
-							if (RelevantFace.bBlendNormals)
-							{
-								CornerTangentZ[CornerIndex] += CornerWeight * TriangleTangentZ[OtherFaceIndex];
-							}
-						}
-					}
-					if (!WedgeTangentX[WedgeOffset + CornerIndex].IsZero())
-					{
-						CornerTangentX[CornerIndex] = WedgeTangentX[WedgeOffset + CornerIndex];
-					}
-					if (!WedgeTangentY[WedgeOffset + CornerIndex].IsZero())
-					{
-						CornerTangentY[CornerIndex] = WedgeTangentY[WedgeOffset + CornerIndex];
-					}
-					if (!WedgeTangentZ[WedgeOffset + CornerIndex].IsZero())
-					{
-						CornerTangentZ[CornerIndex] = WedgeTangentZ[WedgeOffset + CornerIndex];
-					}
-				}
-			}
-
-			// Normalization.
-			for (int32 CornerIndex = 0; CornerIndex < 3; CornerIndex++)
-			{
-				CornerTangentX[CornerIndex].Normalize();
-				CornerTangentY[CornerIndex].Normalize();
-				CornerTangentZ[CornerIndex].Normalize();
-
-				// Gram-Schmidt orthogonalization
-				CornerTangentY[CornerIndex] -= CornerTangentX[CornerIndex] * (CornerTangentX[CornerIndex] | CornerTangentY[CornerIndex]);
-				CornerTangentY[CornerIndex].Normalize();
-
-				CornerTangentX[CornerIndex] -= CornerTangentZ[CornerIndex] * (CornerTangentZ[CornerIndex] | CornerTangentX[CornerIndex]);
-				CornerTangentX[CornerIndex].Normalize();
-				CornerTangentY[CornerIndex] -= CornerTangentZ[CornerIndex] * (CornerTangentZ[CornerIndex] | CornerTangentY[CornerIndex]);
-				CornerTangentY[CornerIndex].Normalize();
-			}
-
-			// Copy back to the mesh.
-			for (int32 CornerIndex = 0; CornerIndex < 3; CornerIndex++)
-			{
-				WedgeTangentX[WedgeOffset + CornerIndex] = CornerTangentX[CornerIndex];
-				WedgeTangentY[WedgeOffset + CornerIndex] = CornerTangentY[CornerIndex];
-				WedgeTangentZ[WedgeOffset + CornerIndex] = CornerTangentZ[CornerIndex];
-			}
-		}
-
-		check(WedgeTangentX.Num() == NumWedges);
-		check(WedgeTangentY.Num() == NumWedges);
-		check(WedgeTangentZ.Num() == NumWedges);
-	}
-
 	bool IsTriangleMirror(IMeshBuildData* BuildData, const TArray<FVector>& TriangleTangentZ, const uint32 FaceIdxA, const uint32 FaceIdxB)
 	{
 		if (FaceIdxA == FaceIdxB)
@@ -3819,7 +3489,7 @@ public:
 		return false;
 	}
 
-	void Skeletal_ComputeTangents_MikkTSpace(
+	void Skeletal_ComputeTangents(
 		IMeshBuildData* BuildData,
 		const FOverlappingCorners& OverlappingCorners
 		)
@@ -3827,6 +3497,7 @@ public:
 		bool bBlendOverlappingNormals = true;
 		bool bIgnoreDegenerateTriangles = BuildData->BuildOptions.bRemoveDegenerateTriangles;
 		bool bComputeWeightedNormals = BuildData->BuildOptions.bComputeWeightedNormals;
+		bool bUseMikktSpace = BuildData->BuildOptions.bUseMikkTSpace && (BuildData->BuildOptions.bComputeNormals || BuildData->BuildOptions.bComputeTangents);
 		
 		int32 NumFaces = BuildData->GetNumFaces();
 		int32 NumWedges = BuildData->GetNumWedges();
@@ -3859,25 +3530,23 @@ public:
 		TArray<FFanFace> RelevantFacesForCorner[3];
 		TArray<int32> AdjacentFaces;
 
-		bool bWedgeTSpace = false;
-
-		if (WedgeTangentX.Num() > 0 && WedgeTangentY.Num() > 0)
+		// Allocate storage for tangents and normal if none were provided.
+		if (WedgeTangentX.Num() != NumWedges)
 		{
-			bWedgeTSpace = true;
-			for (int32 WedgeIdx = 0; WedgeIdx < WedgeTangentX.Num()
-				&& WedgeIdx < WedgeTangentY.Num(); ++WedgeIdx)
-			{
-				bWedgeTSpace = bWedgeTSpace && (!WedgeTangentX[WedgeIdx].IsNearlyZero()) && (!WedgeTangentY[WedgeIdx].IsNearlyZero());
-			}
+			WedgeTangentX.Empty(NumWedges);
+			WedgeTangentX.AddZeroed(NumWedges);
+		}
+		if (WedgeTangentY.Num() != NumWedges)
+		{
+			WedgeTangentY.Empty(NumWedges);
+			WedgeTangentY.AddZeroed(NumWedges);
 		}
-
-		// Allocate storage for tangents if none were provided, and calculate normals for MikkTSpace.
 		if (WedgeTangentZ.Num() != NumWedges)
 		{
-			// normals are not included, so we should calculate them
 			WedgeTangentZ.Empty(NumWedges);
 			WedgeTangentZ.AddZeroed(NumWedges);
 		}
+
 		bool bIsZeroLengthNormalErrorMessageDisplayed = false;
 		// we need to calculate normals for MikkTSpace
 		for (int32 FaceIndex = 0; FaceIndex < NumFaces; FaceIndex++)
@@ -3885,10 +3554,14 @@ public:
 			int32 PatchIndex = FaceIndexToPatchIndex[FaceIndex];
 			int32 WedgeOffset = FaceIndex * 3;
 			FVector CornerPositions[3];
+			FVector CornerTangentX[3];
+			FVector CornerTangentY[3];
 			FVector CornerNormal[3];
 
 			for (int32 CornerIndex = 0; CornerIndex < 3; CornerIndex++)
 			{
+				CornerTangentX[CornerIndex] = FVector::ZeroVector;
+				CornerTangentY[CornerIndex] = FVector::ZeroVector;
 				CornerNormal[CornerIndex] = FVector::ZeroVector;
 				CornerPositions[CornerIndex] = BuildData->GetVertexPosition(FaceIndex, CornerIndex);
 				RelevantFacesForCorner[CornerIndex].Reset();
@@ -3904,15 +3577,31 @@ public:
 
 			// No need to process triangles if tangents already exist.
 			bool bCornerHasNormal[3] = { 0 };
+			bool bCornerHasTangents[3] = { 0 };
+			bool bSkipNTB = true;
 			for (int32 CornerIndex = 0; CornerIndex < 3; CornerIndex++)
 			{
 				bCornerHasNormal[CornerIndex] = !WedgeTangentZ[WedgeOffset + CornerIndex].IsZero();
+				bCornerHasTangents[CornerIndex] = !WedgeTangentX[WedgeOffset + CornerIndex].IsZero() && !WedgeTangentY[WedgeOffset + CornerIndex].IsZero();
+
+				//If we want to compute mikkt we dont check tangents to skip this corner
+				if (!bCornerHasNormal[CornerIndex] || (!bUseMikktSpace && !bCornerHasTangents[CornerIndex]))
+				{
+					bSkipNTB = false;
+				}
 			}
-			if (bCornerHasNormal[0] && bCornerHasNormal[1] && bCornerHasNormal[2])
+			if (bSkipNTB)
 			{
 				continue;
 			}
 
+			// Calculate smooth vertex normals.
+			float Determinant = FVector::Triple(
+				TriangleTangentX[FaceIndex],
+				TriangleTangentY[FaceIndex],
+				TriangleTangentZ[FaceIndex]
+			);
+
 			// Start building a list of faces adjacent to this face.
 			AdjacentFaces.Reset();
 			for (int32 CornerIndex = 0; CornerIndex < 3; CornerIndex++)
@@ -3946,7 +3635,7 @@ public:
 				int32 OtherFaceIndex = AdjacentFaces[AdjacentFaceIndex];
 				for (int32 OurCornerIndex = 0; OurCornerIndex < 3; OurCornerIndex++)
 				{
-					if (bCornerHasNormal[OurCornerIndex])
+					if (bCornerHasNormal[OurCornerIndex] && (bUseMikktSpace || bCornerHasTangents[OurCornerIndex]))
 						continue;
 
 					FFanFace NewFanFace;
@@ -3991,7 +3680,7 @@ public:
 			// smoothing-group-compatible connected triangle fans around common vertices.
 			for (int32 CornerIndex = 0; CornerIndex < 3; CornerIndex++)
 			{
-				if (bCornerHasNormal[CornerIndex])
+				if (bCornerHasNormal[CornerIndex] && (bUseMikktSpace || bCornerHasTangents[CornerIndex]))
 					continue;
 
 				int32 NewConnections;
@@ -4013,6 +3702,7 @@ public:
 											&& (BuildData->GetFaceSmoothingGroups(NextFace.FaceIndex) & BuildData->GetFaceSmoothingGroups(OtherFace.FaceIndex)))
 									{
 										int32 CommonVertices = 0;
+										int32 CommonTangentVertices = 0;
 										int32 CommonNormalVertices = 0;
 										for (int32 OtherCornerIndex = 0; OtherCornerIndex < 3; OtherCornerIndex++)
 										{
@@ -4026,6 +3716,13 @@ public:
 													BuildData->BuildOptions.OverlappingThresholds))
 												{
 													CommonVertices++;
+													if (!bUseMikktSpace && UVsEqual(
+														BuildData->GetVertexUV(NextFace.FaceIndex, NextCornerIndex, 0),
+														BuildData->GetVertexUV(OtherFace.FaceIndex, OtherCornerIndex, 0),
+														BuildData->BuildOptions.OverlappingThresholds))
+													{
+														CommonTangentVertices++;
+													}
 													if (bBlendOverlappingNormals
 														|| NextVertexIndex == OtherVertexIndex)
 													{
@@ -4040,6 +3737,19 @@ public:
 											NextFace.bFilled = true;
 											NextFace.bBlendNormals = (CommonNormalVertices > 1);
 											NewConnections++;
+											// Only blend tangents if there is no UV seam along the edge with this face.
+											if (!bUseMikktSpace && OtherFace.bBlendTangents && CommonTangentVertices > 1)
+											{
+												float OtherDeterminant = FVector::Triple(
+													TriangleTangentX[NextFace.FaceIndex],
+													TriangleTangentY[NextFace.FaceIndex],
+													TriangleTangentZ[NextFace.FaceIndex]
+												);
+												if ((Determinant * OtherDeterminant) > 0.0f)
+												{
+													NextFace.bBlendTangents = true;
+												}
+											}
 										}
 									}
 								}
@@ -4053,11 +3763,19 @@ public:
 			// Vertex normal construction.
 			for (int32 CornerIndex = 0; CornerIndex < 3; CornerIndex++)
 			{
-				if (bCornerHasNormal[CornerIndex])
+				bool bUseProvidedNormal = bCornerHasNormal[CornerIndex];
+				bool bUseProvidedTangents = !bUseMikktSpace && bCornerHasTangents[CornerIndex];
+				if (bUseProvidedNormal)
 				{
 					CornerNormal[CornerIndex] = WedgeTangentZ[WedgeOffset + CornerIndex];
 				}
-				else
+				if (bUseProvidedTangents)
+				{
+					CornerTangentX[CornerIndex] = WedgeTangentX[WedgeOffset + CornerIndex];
+					CornerTangentY[CornerIndex] = WedgeTangentY[WedgeOffset + CornerIndex];
+				}
+
+				if (!bUseProvidedNormal || !bUseProvidedTangents)
 				{
 					for (int32 RelevantFaceIdx = 0; RelevantFaceIdx < RelevantFacesForCorner[CornerIndex].Num(); RelevantFaceIdx++)
 					{
@@ -4075,13 +3793,25 @@ public:
 								//Get the CornerWeight
 								CornerWeight = OtherFaceArea * OtherFaceAngle;
 							}
-							
-							if (RelevantFace.bBlendNormals)
+							if (!bUseProvidedTangents && RelevantFace.bBlendTangents)
+							{
+								CornerTangentX[CornerIndex] += CornerWeight * TriangleTangentX[OtherFaceIndex];
+								CornerTangentY[CornerIndex] += CornerWeight * TriangleTangentY[OtherFaceIndex];
+							}
+							if (!bUseProvidedNormal && RelevantFace.bBlendNormals)
 							{
 								CornerNormal[CornerIndex] += CornerWeight * TriangleTangentZ[OtherFaceIndex];
 							}
 						}
 					}
+					if (!WedgeTangentX[WedgeOffset + CornerIndex].IsZero())
+					{
+						CornerTangentX[CornerIndex] = WedgeTangentX[WedgeOffset + CornerIndex];
+					}
+					if (!WedgeTangentY[WedgeOffset + CornerIndex].IsZero())
+					{
+						CornerTangentY[CornerIndex] = WedgeTangentY[WedgeOffset + CornerIndex];
+					}
 					if (!WedgeTangentZ[WedgeOffset + CornerIndex].IsZero())
 					{
 						CornerNormal[CornerIndex] = WedgeTangentZ[WedgeOffset + CornerIndex];
@@ -4093,15 +3823,27 @@ public:
 			for (int32 CornerIndex = 0; CornerIndex < 3; CornerIndex++)
 			{
 				CornerNormal[CornerIndex].Normalize();
+				if (!bUseMikktSpace)
+				{
+					CornerTangentX[CornerIndex].Normalize();
+					CornerTangentY[CornerIndex].Normalize();
+
+					// Gram-Schmidt orthogonalization
+					CornerTangentY[CornerIndex] -= CornerTangentX[CornerIndex] * (CornerTangentX[CornerIndex] | CornerTangentY[CornerIndex]);
+					CornerTangentY[CornerIndex].Normalize();
+
+					CornerTangentX[CornerIndex] -= CornerNormal[CornerIndex] * (CornerNormal[CornerIndex] | CornerTangentX[CornerIndex]);
+					CornerTangentX[CornerIndex].Normalize();
+					CornerTangentY[CornerIndex] -= CornerNormal[CornerIndex] * (CornerNormal[CornerIndex] | CornerTangentY[CornerIndex]);
+					CornerTangentY[CornerIndex].Normalize();
+				}
 			}
 
-			// Copy back to the mesh.
-			for (int32 CornerIndex = 0; CornerIndex < 3; CornerIndex++)
+			auto VerifyTangentSpace = [&bIsZeroLengthNormalErrorMessageDisplayed, &BuildData](FVector& NormalizedVector)
 			{
-				//Make sure the normal do not contain nan, if it contain nan change it to a valid zero vector
-				if (CornerNormal[CornerIndex].IsNearlyZero() || CornerNormal[CornerIndex].ContainsNaN())
+				if (NormalizedVector.IsNearlyZero() || NormalizedVector.ContainsNaN())
 				{
-					CornerNormal[CornerIndex] = FVector::ZeroVector;
+					NormalizedVector = FVector::ZeroVector;
 					//We also notify the log that we compute a zero length normals, so the user is aware of it
 					if (!bIsZeroLengthNormalErrorMessageDisplayed)
 					{
@@ -4122,22 +3864,26 @@ public:
 						}
 					}
 				}
+			};
+
+			// Copy back to the mesh.
+			for (int32 CornerIndex = 0; CornerIndex < 3; CornerIndex++)
+			{
+				VerifyTangentSpace(CornerNormal[CornerIndex]);
 				WedgeTangentZ[WedgeOffset + CornerIndex] = CornerNormal[CornerIndex];
-			}
-		}
 
-		if (WedgeTangentX.Num() != NumWedges)
-		{
-			WedgeTangentX.Empty(NumWedges);
-			WedgeTangentX.AddZeroed(NumWedges);
-		}
-		if (WedgeTangentY.Num() != NumWedges)
-		{
-			WedgeTangentY.Empty(NumWedges);
-			WedgeTangentY.AddZeroed(NumWedges);
+				if (!bUseMikktSpace)
+				{
+					VerifyTangentSpace(CornerTangentX[CornerIndex]);
+					WedgeTangentX[WedgeOffset + CornerIndex] = CornerTangentX[CornerIndex];
+
+					VerifyTangentSpace(CornerTangentY[CornerIndex]);
+					WedgeTangentY[WedgeOffset + CornerIndex] = CornerTangentY[CornerIndex];
+				}
+			}
 		}
 
-		//if (!bWedgeTSpace)
+		if (bUseMikktSpace)
 		{
 			// we can use mikktspace to calculate the tangents		
 			SMikkTSpaceContext MikkTContext;
@@ -4190,15 +3936,8 @@ public:
 			TangentZ.AddZeroed(NumWedges);
 		}
 
-		// Compute any missing tangents. MikkTSpace should be use only when the user want to recompute the normals or tangents otherwise should always fallback on builtin
-		if (BuildData->BuildOptions.bUseMikkTSpace && (BuildData->BuildOptions.bComputeNormals || BuildData->BuildOptions.bComputeTangents))
-		{
-			Skeletal_ComputeTangents_MikkTSpace(BuildData, OverlappingCorners);
-		}
-		else
-		{
-			Skeletal_ComputeTangents(BuildData, OverlappingCorners);
-		}
+		// Compute any missing tangents. MikkTSpace should be use only when the user want to recompute the normals or tangents otherwise should always fallback on builtin tangent
+		Skeletal_ComputeTangents(BuildData, OverlappingCorners);
 
 		// At this point the mesh will have valid tangents.
 		BuildData->ValidateTangentArraySize();
@@ -4477,18 +4216,6 @@ bool FMeshUtilities::BuildSkeletalMesh(FSkeletalMeshLODModel& LODModel, const FR
 		}
 	};
 
-	// Temporarily supporting both import paths
-	if (!BuildOptions.bUseMikkTSpace)
-	{
-		bool bBuildSuccess = BuildSkeletalMesh_Legacy(LODModel, RefSkeleton, Influences, Wedges, Faces, Points, PointToOriginalMap, BuildOptions.OverlappingThresholds, BuildOptions.bComputeNormals, BuildOptions.bComputeTangents, BuildOptions.bComputeWeightedNormals, OutWarningMessages, OutWarningNames);
-		if (bBuildSuccess)
-		{
-			UpdateOverlappingVertices(LODModel);
-		}
-
-		return bBuildSuccess;
-	}
-
 	SkeletalMeshBuildData BuildData(
 		LODModel,
 		RefSkeleton,
diff --git a/Engine/Source/Developer/MeshUtilities/Private/MeshUtilitiesPrivate.h b/Engine/Source/Developer/MeshUtilities/Private/MeshUtilitiesPrivate.h
--- a/Engine/Source/Developer/MeshUtilities/Private/MeshUtilitiesPrivate.h
+++ b/Engine/Source/Developer/MeshUtilities/Private/MeshUtilitiesPrivate.h
@@ -109,6 +109,8 @@ private:
 	virtual bool GenerateUniqueUVsForStaticMesh(const FRawMesh& RawMesh, int32 TextureResolution, bool bMergeIdenticalMaterials, TArray<FVector2D>& OutTexCoords) const override;
 
 	virtual bool BuildSkeletalMesh(FSkeletalMeshLODModel& LODModel, const FReferenceSkeleton& RefSkeleton, const TArray<SkeletalMeshImportData::FVertInfluence>& Influences, const TArray<SkeletalMeshImportData::FMeshWedge>& Wedges, const TArray<SkeletalMeshImportData::FMeshFace>& Faces, const TArray<FVector>& Points, const TArray<int32>& PointToOriginalMap, const MeshBuildOptions& BuildOptions = MeshBuildOptions(), TArray<FText> * OutWarningMessages = NULL, TArray<FName> * OutWarningNames = NULL) override;
+	
+	UE_DEPRECATED(4.24, "Use functionality in FSkeletalMeshUtilityBuilder instead.")
 	bool BuildSkeletalMesh_Legacy(FSkeletalMeshLODModel& LODModel, const FReferenceSkeleton& RefSkeleton, const TArray<SkeletalMeshImportData::FVertInfluence>& Influences, const TArray<SkeletalMeshImportData::FMeshWedge>& Wedges, const TArray<SkeletalMeshImportData::FMeshFace>& Faces, const TArray<FVector>& Points, const TArray<int32>& PointToOriginalMap, const FOverlappingThresholds& OverlappingThresholds, bool bComputeNormals = true, bool bComputeTangents = true, bool bComputeWeightedNormals = true, TArray<FText> * OutWarningMessages = NULL, TArray<FName> * OutWarningNames = NULL);
 
 	virtual void CacheOptimizeIndexBuffer(TArray<uint16>& Indices) override;
diff --git a/Engine/Source/Developer/MeshUtilities/Private/SkeletalMeshTools.cpp b/Engine/Source/Developer/MeshUtilities/Private/SkeletalMeshTools.cpp
--- a/Engine/Source/Developer/MeshUtilities/Private/SkeletalMeshTools.cpp
+++ b/Engine/Source/Developer/MeshUtilities/Private/SkeletalMeshTools.cpp
@@ -207,8 +207,8 @@ namespace SkeletalMeshTools
 		for (int32 SrcChunkIndex = 0; SrcChunkIndex < SrcChunks.Num(); ++SrcChunkIndex)
 		{
 			FSkinnedMeshChunk* SrcChunk = SrcChunks[SrcChunkIndex];
+			SrcChunk->OriginalSectionIndex = SrcChunkIndex;
 			int32 FirstChunkIndex = Chunks.Num();
-
 			for (int32 i = 0; i < SrcChunk->Indices.Num(); i += 3)
 			{
 				// Find all bones needed by this triangle.
@@ -257,9 +257,10 @@ namespace SkeletalMeshTools
 				if (DestChunk == NULL)
 				{
 					DestChunk = new FSkinnedMeshChunk();
-					Chunks.Add(DestChunk);
+					int32 ChunkSectionIndex = Chunks.Add(DestChunk);
 					DestChunk->MaterialIndex = SrcChunk->MaterialIndex;
 					DestChunk->OriginalSectionIndex = SrcChunk->OriginalSectionIndex;
+					DestChunk->ParentChunkSectionIndex = ChunkSectionIndex == FirstChunkIndex ? INDEX_NONE : FirstChunkIndex;
 					TArray<int32>& IndexMap = *new TArray<int32>();
 					IndexMaps.Add(&IndexMap);
 					IndexMap.AddUninitialized(SrcChunk->Vertices.Num());
diff --git a/Engine/Source/Developer/MeshUtilities/Private/SkeletalMeshTools.h b/Engine/Source/Developer/MeshUtilities/Private/SkeletalMeshTools.h
--- a/Engine/Source/Developer/MeshUtilities/Private/SkeletalMeshTools.h
+++ b/Engine/Source/Developer/MeshUtilities/Private/SkeletalMeshTools.h
@@ -51,6 +51,8 @@ struct FSkinnedMeshChunk
 	TArray<uint32> Indices;
 	/** If not empty, contains a map from bones referenced in this chunk to the skeleton. */
 	TArray<FBoneIndexType> BoneMap;
+	/** The parent original section index for which this chunk was generated. INDEX_NONE for parent the value of the parent for BONE child chunked*/
+	int32 ParentChunkSectionIndex;
 };
 
 /**
diff --git a/Engine/Source/Developer/MeshUtilities/Public/MeshUtilities.h b/Engine/Source/Developer/MeshUtilities/Public/MeshUtilities.h
--- a/Engine/Source/Developer/MeshUtilities/Public/MeshUtilities.h
+++ b/Engine/Source/Developer/MeshUtilities/Public/MeshUtilities.h
@@ -228,6 +228,18 @@ public:
 		bool bUseMikkTSpace;
 		bool bComputeWeightedNormals;
 		FOverlappingThresholds OverlappingThresholds;
+
+		void FillOptions(const FSkeletalMeshBuildSettings& SkeletalMeshBuildSettings)
+		{
+			OverlappingThresholds.ThresholdPosition = SkeletalMeshBuildSettings.ThresholdPosition;
+			OverlappingThresholds.ThresholdTangentNormal = SkeletalMeshBuildSettings.ThresholdTangentNormal;
+			OverlappingThresholds.ThresholdUV = SkeletalMeshBuildSettings.ThresholdUV;
+			bComputeNormals = SkeletalMeshBuildSettings.bRecomputeNormals;
+			bComputeTangents = SkeletalMeshBuildSettings.bRecomputeTangents;
+			bUseMikkTSpace = SkeletalMeshBuildSettings.bUseMikkTSpace;
+			bComputeWeightedNormals = SkeletalMeshBuildSettings.bComputeWeightedNormals;
+			bRemoveDegenerateTriangles = SkeletalMeshBuildSettings.bRemoveDegenerates;
+		}
 	};
 	
 	/**
diff --git a/Engine/Source/Developer/SkeletalMeshUtilitiesCommon/Private/LODUtilities.cpp b/Engine/Source/Developer/SkeletalMeshUtilitiesCommon/Private/LODUtilities.cpp
--- a/Engine/Source/Developer/SkeletalMeshUtilitiesCommon/Private/LODUtilities.cpp
+++ b/Engine/Source/Developer/SkeletalMeshUtilitiesCommon/Private/LODUtilities.cpp
@@ -202,14 +202,11 @@ bool FLODUtilities::RegenerateLOD(USkeletalMesh* SkeletalMesh, int32 NewLODCount
 {
 	if (SkeletalMesh)
 	{
+		FScopedSkeletalMeshPostEditChange ScopedPostEditChange(SkeletalMesh);
+
 		// Unbind any existing clothing assets before we regenerate all LODs
 		TArray<ClothingAssetUtils::FClothingAssetMeshBinding> ClothingBindings;
-		ClothingAssetUtils::GetMeshClothingAssetBindings(SkeletalMesh, ClothingBindings);
-
-		for (ClothingAssetUtils::FClothingAssetMeshBinding& Binding : ClothingBindings)
-		{
-			Binding.Asset->UnbindFromSkeletalMesh(SkeletalMesh, Binding.LODIndex);
-		}
+		FLODUtilities::UnbindClothingAndBackup(SkeletalMesh, ClothingBindings);
 
 		int32 LODCount = SkeletalMesh->GetLODNum();
 
@@ -239,7 +236,7 @@ bool FLODUtilities::RegenerateLOD(USkeletalMesh* SkeletalMesh, int32 NewLODCount
 			for (int32 LODIdx = CurrentNumLODs; LODIdx < LODCount; LODIdx++)
 			{
 				// if no previous setting found, it will use default setting. 
-				FLODUtilities::SimplifySkeletalMeshLOD(UpdateContext, LODIdx);
+				FLODUtilities::SimplifySkeletalMeshLOD(UpdateContext, LODIdx, false);
 			}
 		}
 		else
@@ -249,22 +246,13 @@ bool FLODUtilities::RegenerateLOD(USkeletalMesh* SkeletalMesh, int32 NewLODCount
 				FSkeletalMeshLODInfo& CurrentLODInfo = *(SkeletalMesh->GetLODInfo(LODIdx));
 				if ((bRegenerateEvenIfImported && LODIdx > 0) || (bGenerateBaseLOD && LODIdx == 0) || CurrentLODInfo.bHasBeenSimplified )
 				{
-					FLODUtilities::SimplifySkeletalMeshLOD(UpdateContext, LODIdx);
+					FLODUtilities::SimplifySkeletalMeshLOD(UpdateContext, LODIdx, false);
 				}
 			}
 		}
 
 		//Restore all clothing we can
-		for (ClothingAssetUtils::FClothingAssetMeshBinding& Binding : ClothingBindings)
-		{
-			if (SkeletalMesh->GetImportedModel()->LODModels.IsValidIndex(Binding.LODIndex) &&
-				SkeletalMesh->GetImportedModel()->LODModels[Binding.LODIndex].Sections.IsValidIndex(Binding.SectionIndex))
-			{
-				Binding.Asset->BindToSkeletalMesh(SkeletalMesh, Binding.LODIndex, Binding.SectionIndex, Binding.AssetInternalLodIndex);
-			}
-		}
-
-		SkeletalMesh->PostEditChange();
+		FLODUtilities::RestoreClothingFromBackup(SkeletalMesh, ClothingBindings);
 
 		return true;
 	}
@@ -296,21 +284,12 @@ void FLODUtilities::RemoveLOD(FSkeletalMeshUpdateContext& UpdateContext, int32 D
 	// If its a valid LOD, kill it.
 	if( DesiredLOD > 0 && DesiredLOD < SkelMeshModel->LODModels.Num() )
 	{
-		//We'll be modifying the skel mesh data so reregister
-
-		//TODO - do we need to reregister something else instead?
-		FMultiComponentReregisterContext ReregisterContext(UpdateContext.AssociatedComponents);
-
-		// Release rendering resources before deleting LOD
-		SkeletalMesh->ReleaseResources();
+		FScopedSkeletalMeshPostEditChange ScopedPostEditChange(SkeletalMesh);
 
 		// Block until this is done
-		FlushRenderingCommands();
 
 		SkelMeshModel->LODModels.RemoveAt(DesiredLOD);
 		SkeletalMesh->RemoveLODInfo(DesiredLOD);
-		SkeletalMesh->InitResources();
-
 		RefreshLODChange(SkeletalMesh);
 
 		// Set the forced LOD to Auto.
@@ -332,9 +311,6 @@ void FLODUtilities::RemoveLOD(FSkeletalMeshUpdateContext& UpdateContext, int32 D
 			}
 		}
 
-		// This will recache derived render data, and re-init resources
-		SkeletalMesh->PostEditChange();
-
 		//Notify calling system of change
 		UpdateContext.OnLODChanged.ExecuteIfBound();
 
@@ -488,7 +464,7 @@ void ProjectTargetOnBase(const TArray<FSoftSkinVertex>& BaseVertices, const TArr
 	{
 		//Use the remap base index in case some sections disappear during the reduce phase
 		int32 BaseSectionIndex = TargetSectionMatchBaseIndex[SectionIndex];
-		if (BaseSectionIndex == INDEX_NONE || !PerSectionBaseTriangleIndices.IsValidIndex(BaseSectionIndex))
+		if (BaseSectionIndex == INDEX_NONE || !PerSectionBaseTriangleIndices.IsValidIndex(BaseSectionIndex) || PerSectionBaseTriangleIndices[BaseSectionIndex].Num() < 1)
 		{
 			continue;
 		}
@@ -518,6 +494,7 @@ void ProjectTargetOnBase(const TArray<FSoftSkinVertex>& BaseVertices, const TArr
 			TriangleElement.TriangleIndex = Triangles.Num();
 			Triangles.Add(TriangleElement);
 		}
+		check(!BaseMeshUVBound.GetExtent().IsNearlyZero());
 		//Setup the Quad tree
 		float UVsQuadTreeMinSize = 0.001f;
 		TQuadTree<uint32, 100> QuadTree(BaseMeshUVBound, UVsQuadTreeMinSize);
@@ -661,7 +638,7 @@ void CreateLODMorphTarget(USkeletalMesh* SkeletalMesh, FReductionBaseSkeletalMes
 	if (ReductionBaseSkeletalMeshBulkData != nullptr)
 	{
 		FSkeletalMeshLODModel TempBaseLODModel;
-		ReductionBaseSkeletalMeshBulkData->LoadReductionData(TempBaseLODModel, BaseLODMorphTargetData);
+		ReductionBaseSkeletalMeshBulkData->LoadReductionData(TempBaseLODModel, BaseLODMorphTargetData, SkeletalMesh);
 	}
 
 	FSkeletalMeshModel* SkeletalMeshModel = SkeletalMesh->GetImportedModel();
@@ -841,7 +818,8 @@ void FLODUtilities::ApplyMorphTargetsToLOD(USkeletalMesh* SkeletalMesh, int32 So
 	TMap<FString, TArray<FMorphTargetDelta>> TempBaseLODMorphTargetData;
 	if (bReduceBaseLOD)
 	{
-		ReductionBaseSkeletalMeshBulkData->LoadReductionData(TempBaseLODModel, TempBaseLODMorphTargetData);
+		check(ReductionBaseSkeletalMeshBulkData != nullptr);
+		ReductionBaseSkeletalMeshBulkData->LoadReductionData(TempBaseLODModel, TempBaseLODMorphTargetData, SkeletalMesh);
 	}
 
 	const FSkeletalMeshLODModel& BaseLODModel = bReduceBaseLOD ? TempBaseLODModel : SkeletalMeshResource->LODModels[SourceLOD];
@@ -900,7 +878,7 @@ void FLODUtilities::ApplyMorphTargetsToLOD(USkeletalMesh* SkeletalMesh, int32 So
 	for (int32 SectionIndex = 0; SectionIndex < SectionCount; ++SectionIndex)
 	{
 		const FSkelMeshSection& Section = BaseLODModel.Sections[SectionIndex];
-		if (Section.HasClothingData())
+		if (Section.ClothingData.AssetGuid.IsValid())
 		{
 			continue;
 		}
@@ -964,13 +942,14 @@ void FLODUtilities::ApplyMorphTargetsToLOD(USkeletalMesh* SkeletalMesh, int32 So
 	CreateLODMorphTarget(SkeletalMesh, ReductionBaseSkeletalMeshBulkData, SourceLOD, DestinationLOD, PerMorphTargetBaseIndexToMorphTargetDelta, BaseMorphIndexToTargetIndexList, TargetVertices, TargetMatchData);
 }
 
-void FLODUtilities::SimplifySkeletalMeshLOD( USkeletalMesh* SkeletalMesh, int32 DesiredLOD, bool bReregisterComponent /*= true*/, bool bRestoreClothing /*= false*/)
+void FLODUtilities::SimplifySkeletalMeshLOD( USkeletalMesh* SkeletalMesh, int32 DesiredLOD, bool bRestoreClothing /*= false*/)
 {
 	IMeshReductionModule& ReductionModule = FModuleManager::Get().LoadModuleChecked<IMeshReductionModule>("MeshReductionInterface");
 	IMeshReduction* MeshReduction = ReductionModule.GetSkeletalMeshReductionInterface();
 
 	check (MeshReduction && MeshReduction->IsSupported());
 
+
 	if (DesiredLOD == 0
 		&& SkeletalMesh->GetLODInfo(DesiredLOD) != nullptr
 		&& SkeletalMesh->GetLODInfo(DesiredLOD)->bHasBeenSimplified
@@ -1000,18 +979,13 @@ void FLODUtilities::SimplifySkeletalMeshLOD( USkeletalMesh* SkeletalMesh, int32
 		GWarn->BeginSlowTask(StatusUpdate, true);
 	}
 
+	FScopedSkeletalMeshPostEditChange ScopedPostEditChange(SkeletalMesh);
+
 	// Unbind DesiredLOD existing clothing assets before we simplify this LOD
 	TArray<ClothingAssetUtils::FClothingAssetMeshBinding> ClothingBindings;
-	if (bRestoreClothing)
+	if (bRestoreClothing && SkeletalMesh->GetImportedModel() && SkeletalMesh->GetImportedModel()->LODModels.IsValidIndex(DesiredLOD))
 	{
-		ClothingAssetUtils::GetMeshClothingAssetBindings(SkeletalMesh, ClothingBindings);
-		for (ClothingAssetUtils::FClothingAssetMeshBinding& Binding : ClothingBindings)
-		{
-			if (DesiredLOD == Binding.LODIndex)
-			{
-				Binding.Asset->UnbindFromSkeletalMesh(SkeletalMesh, Binding.LODIndex);
-			}
-		}
+		FLODUtilities::UnbindClothingAndBackup(SkeletalMesh, ClothingBindings, DesiredLOD);
 	}
 
 	if (SkeletalMesh->GetLODInfo(DesiredLOD) != nullptr)
@@ -1048,8 +1022,20 @@ void FLODUtilities::SimplifySkeletalMeshLOD( USkeletalMesh* SkeletalMesh, int32
 						MorphDeltasArray.Add(MorphDelta);
 					}
 				}
+				
 				//Copy the original SkeletalMesh LODModel
-				SkeletalMeshResource->OriginalReductionSourceMeshData[DesiredLOD]->SaveReductionData(SrcModel, BaseLODMorphTargetData);
+				// Unbind clothing before saving the original data, we must not restore clothing to do inline reduction
+				{
+					TArray<ClothingAssetUtils::FClothingAssetMeshBinding> TemporaryRemoveClothingBindings;
+					FLODUtilities::UnbindClothingAndBackup(SkeletalMesh, TemporaryRemoveClothingBindings, DesiredLOD);
+
+					SkeletalMeshResource->OriginalReductionSourceMeshData[DesiredLOD]->SaveReductionData(SrcModel, BaseLODMorphTargetData, SkeletalMesh);
+
+					if (TemporaryRemoveClothingBindings.Num() > 0)
+					{
+						FLODUtilities::RestoreClothingFromBackup(SkeletalMesh, TemporaryRemoveClothingBindings, DesiredLOD);
+					}
+				}
 
 				if (DesiredLOD == 0)
 				{
@@ -1060,7 +1046,7 @@ void FLODUtilities::SimplifySkeletalMeshLOD( USkeletalMesh* SkeletalMesh, int32
 	}
 	
 
-	if (MeshReduction->ReduceSkeletalMesh(SkeletalMesh, DesiredLOD, bReregisterComponent))
+	if (MeshReduction->ReduceSkeletalMesh(SkeletalMesh, DesiredLOD))
 	{
 		check(SkeletalMesh->GetLODNum() >= 1);
 
@@ -1078,21 +1064,7 @@ void FLODUtilities::SimplifySkeletalMeshLOD( USkeletalMesh* SkeletalMesh, int32
 			}
 		};
 
-		if (bReregisterComponent)
-		{
-			TComponentReregisterContext<USkinnedMeshComponent> ReregisterContext;
-			SkeletalMesh->ReleaseResources();
-			SkeletalMesh->ReleaseResourcesFence.Wait();
-
-			ApplyMorphTargetOption();
-			
-			SkeletalMesh->PostEditChange();
-			SkeletalMesh->InitResources();
-		}
-		else
-		{
-			ApplyMorphTargetOption();
-		}
+		ApplyMorphTargetOption();
 		SkeletalMesh->MarkPackageDirty();
 	}
 	else
@@ -1105,25 +1077,15 @@ void FLODUtilities::SimplifySkeletalMeshLOD( USkeletalMesh* SkeletalMesh, int32
 	}
 
 	//Put back the clothing for the DesiredLOD
-	if (bRestoreClothing)
+	if (bRestoreClothing && ClothingBindings.Num() > 0 && SkeletalMesh->GetImportedModel()->LODModels.IsValidIndex(DesiredLOD))
 	{
-		for (ClothingAssetUtils::FClothingAssetMeshBinding& Binding : ClothingBindings)
-		{
-			if (SkeletalMesh->GetImportedModel()->LODModels.IsValidIndex(Binding.LODIndex) &&
-				SkeletalMesh->GetImportedModel()->LODModels[Binding.LODIndex].Sections.IsValidIndex(Binding.SectionIndex))
-			{
-				if (DesiredLOD == Binding.LODIndex)
-				{
-					Binding.Asset->BindToSkeletalMesh(SkeletalMesh, Binding.LODIndex, Binding.SectionIndex, Binding.AssetInternalLodIndex);
-				}
-			}
-		}
+		FLODUtilities::RestoreClothingFromBackup(SkeletalMesh, ClothingBindings, DesiredLOD);
 	}
 
 	GWarn->EndSlowTask();
 }
 
-void FLODUtilities::SimplifySkeletalMeshLOD(FSkeletalMeshUpdateContext& UpdateContext, int32 DesiredLOD, bool bReregisterComponent /*= true*/, bool bRestoreClothing /*= false*/)
+void FLODUtilities::SimplifySkeletalMeshLOD(FSkeletalMeshUpdateContext& UpdateContext, int32 DesiredLOD, bool bRestoreClothing /*= false*/)
 {
 	USkeletalMesh* SkeletalMesh = UpdateContext.SkeletalMesh;
 	IMeshReductionModule& ReductionModule = FModuleManager::Get().LoadModuleChecked<IMeshReductionModule>("MeshReductionInterface");
@@ -1131,7 +1093,7 @@ void FLODUtilities::SimplifySkeletalMeshLOD(FSkeletalMeshUpdateContext& UpdateCo
 
 	if (MeshReduction && MeshReduction->IsSupported() && SkeletalMesh)
 	{
-		SimplifySkeletalMeshLOD(SkeletalMesh, DesiredLOD, bReregisterComponent, bRestoreClothing);
+		SimplifySkeletalMeshLOD(SkeletalMesh, DesiredLOD, bRestoreClothing);
 		
 		if (UpdateContext.OnLODChanged.IsBound())
 		{
@@ -1141,40 +1103,27 @@ void FLODUtilities::SimplifySkeletalMeshLOD(FSkeletalMeshUpdateContext& UpdateCo
 	}
 }
 
-void FLODUtilities::RestoreSkeletalMeshLODImportedData(USkeletalMesh* SkeletalMesh, int32 LodIndex, bool bReregisterComponent /*= true*/)
+void FLODUtilities::RestoreSkeletalMeshLODImportedData(USkeletalMesh* SkeletalMesh, int32 LodIndex)
 {
-
 	if (!SkeletalMesh->GetImportedModel()->OriginalReductionSourceMeshData.IsValidIndex(LodIndex) || SkeletalMesh->GetImportedModel()->OriginalReductionSourceMeshData[LodIndex]->IsEmpty())
 	{
 		//There is nothing to restore
 		return;
 	}
 
+	FScopedSkeletalMeshPostEditChange ScopedPostEditChange(SkeletalMesh);
+
 	// Unbind LodIndex existing clothing assets before restoring the LOD
 	TArray<ClothingAssetUtils::FClothingAssetMeshBinding> ClothingBindings;
-	ClothingAssetUtils::GetMeshClothingAssetBindings(SkeletalMesh, ClothingBindings);
-	for (ClothingAssetUtils::FClothingAssetMeshBinding& Binding : ClothingBindings)
-	{
-		//Unbind only the LOD we restore
-		if (Binding.LODIndex == LodIndex)
-		{
-			Binding.Asset->UnbindFromSkeletalMesh(SkeletalMesh, Binding.LODIndex);
-		}
-	}
+	FLODUtilities::UnbindClothingAndBackup(SkeletalMesh, ClothingBindings);
 
 	FSkeletalMeshLODModel ImportedBaseLODModel;
 	TMap<FString, TArray<FMorphTargetDelta>> ImportedBaseLODMorphTargetData;
-	SkeletalMesh->GetImportedModel()->OriginalReductionSourceMeshData[LodIndex]->LoadReductionData(ImportedBaseLODModel, ImportedBaseLODMorphTargetData);
+	SkeletalMesh->GetImportedModel()->OriginalReductionSourceMeshData[LodIndex]->LoadReductionData(ImportedBaseLODModel, ImportedBaseLODMorphTargetData, SkeletalMesh);
 	{
 		FSkeletalMeshUpdateContext UpdateContext;
 		UpdateContext.SkeletalMesh = SkeletalMesh;
 
-		TComponentReregisterContext<USkinnedMeshComponent> ReregisterContext;
-		if (bReregisterComponent)
-		{
-			SkeletalMesh->ReleaseResources();
-			SkeletalMesh->ReleaseResourcesFence.Wait();
-		}
 		//Copy the SkeletalMeshLODModel
 		SkeletalMesh->GetImportedModel()->LODModels[LodIndex] = ImportedBaseLODModel;
 		//Copy the morph target deltas
@@ -1196,20 +1145,7 @@ void FLODUtilities::RestoreSkeletalMeshLODImportedData(USkeletalMesh* SkeletalMe
 		SkeletalMesh->GetImportedModel()->OriginalReductionSourceMeshData[LodIndex]->EmptyBulkData();
 
 		//Put back the clothing for the restore LOD
-		for (ClothingAssetUtils::FClothingAssetMeshBinding& Binding : ClothingBindings)
-		{
-			if (LodIndex == Binding.LODIndex && SkeletalMesh->GetImportedModel()->LODModels.IsValidIndex(Binding.LODIndex) &&
-				SkeletalMesh->GetImportedModel()->LODModels[Binding.LODIndex].Sections.IsValidIndex(Binding.SectionIndex))
-			{
-				Binding.Asset->BindToSkeletalMesh(SkeletalMesh, Binding.LODIndex, Binding.SectionIndex, Binding.AssetInternalLodIndex);
-			}
-		}
-
-		if (bReregisterComponent)
-		{
-			SkeletalMesh->PostEditChange();
-			SkeletalMesh->InitResources();
-		}
+		FLODUtilities::RestoreClothingFromBackup(SkeletalMesh, ClothingBindings);
 
 		if (UpdateContext.OnLODChanged.IsBound())
 		{
@@ -1600,17 +1536,21 @@ void MatchVertexIndexUsingPosition(
 
 bool FLODUtilities::UpdateAlternateSkinWeights(USkeletalMesh* SkeletalMeshDest, const FName& ProfileNameDest, int32 LODIndexDest, FOverlappingThresholds OverlappingThresholds, bool ShouldImportNormals, bool ShouldImportTangents, bool bUseMikkTSpace, bool bComputeWeightedNormals)
 {
-	//Ensure log message only once
-	bool bNoMatchMsgDone = false;
-
 	//Grab all the destination structure
 	check(SkeletalMeshDest);
 	check(SkeletalMeshDest->GetImportedModel());
 	check(SkeletalMeshDest->GetImportedModel()->LODModels.IsValidIndex(LODIndexDest));
 	FSkeletalMeshLODModel& LODModelDest = SkeletalMeshDest->GetImportedModel()->LODModels[LODIndexDest];
+	return UpdateAlternateSkinWeights(LODModelDest, SkeletalMeshDest->GetName(), SkeletalMeshDest->RefSkeleton, ProfileNameDest, LODIndexDest, OverlappingThresholds, ShouldImportNormals, ShouldImportTangents, bUseMikkTSpace, bComputeWeightedNormals);
+}
+
+bool FLODUtilities::UpdateAlternateSkinWeights(FSkeletalMeshLODModel& LODModelDest, const FString SkeletalMeshName, FReferenceSkeleton& RefSkeleton, const FName& ProfileNameDest, int32 LODIndexDest, FOverlappingThresholds OverlappingThresholds, bool ShouldImportNormals, bool ShouldImportTangents, bool bUseMikkTSpace, bool bComputeWeightedNormals)
+{
+	//Ensure log message only once
+	bool bNoMatchMsgDone = false;
 	if (LODModelDest.RawSkeletalMeshBulkData.IsEmpty())
 	{
-		UE_LOG(LogLODUtilities, Error, TEXT("Failed to import Skin Weight Profile as the target skeletal mesh (%s) requires reimporting first."), SkeletalMeshDest ? *SkeletalMeshDest->GetName() : TEXT("NULL"));
+		UE_LOG(LogLODUtilities, Error, TEXT("Failed to import Skin Weight Profile as the target skeletal mesh (%s) requires reimporting first."), *SkeletalMeshName);
 		//Very old asset will not have this data, we cannot add alternate until the asset is reimported
 		return false;
 	}
@@ -1622,7 +1562,7 @@ bool FLODUtilities::UpdateAlternateSkinWeights(USkeletalMesh* SkeletalMeshDest,
 	int32 ProfileIndex = 0;
 	if (!ImportDataDest.AlternateInfluenceProfileNames.Find(ProfileNameDest.ToString(), ProfileIndex))
 	{
-		UE_LOG(LogLODUtilities, Error, TEXT("Failed to import Skin Weight Profile the alternate skinning imported source data is not available."), SkeletalMeshDest ? *SkeletalMeshDest->GetName() : TEXT("NULL"));
+		UE_LOG(LogLODUtilities, Error, TEXT("Failed to import Skin Weight Profile the alternate skinning imported source data is not available."), *SkeletalMeshName);
 		return false;
 	}
 	check(ImportDataDest.AlternateInfluences.IsValidIndex(ProfileIndex));
@@ -1634,7 +1574,7 @@ bool FLODUtilities::UpdateAlternateSkinWeights(USkeletalMesh* SkeletalMeshDest,
 
 	if (ImportDataDest.NumTexCoords <= 0 || ImportDataSrc.NumTexCoords <= 0)
 	{
-		UE_LOG(LogLODUtilities, Error, TEXT("Failed to import Skin Weight Profile as the target skeletal mesh (%s) or imported file does not contain UV coordinates."), SkeletalMeshDest ? *SkeletalMeshDest->GetName() : TEXT("NULL"));
+		UE_LOG(LogLODUtilities, Error, TEXT("Failed to import Skin Weight Profile as the target skeletal mesh (%s) or imported file does not contain UV coordinates."), *SkeletalMeshName);
 		return false;
 	}
 
@@ -1650,7 +1590,7 @@ bool FLODUtilities::UpdateAlternateSkinWeights(USkeletalMesh* SkeletalMeshDest,
 
 	//Create a map to remap source bone index to destination bone index
 	TMap<int32, int32> RemapBoneIndexSrcToDest;
-	FillRemapBoneIndexSrcToDest(ImportDataSrc, ImportDataDest, SkeletalMeshDest->GetName(), LODIndexDest, RemapBoneIndexSrcToDest);
+	FillRemapBoneIndexSrcToDest(ImportDataSrc, ImportDataDest, SkeletalMeshName, LODIndexDest, RemapBoneIndexSrcToDest);
 
 	//Map to get the vertex index source to a destination vertex match
 	TSortedMap<uint32, VertexMatchNameSpace::FVertexMatchResult> VertexIndexSrcToVertexIndexDestMatches;
@@ -1879,7 +1819,10 @@ bool FLODUtilities::UpdateAlternateSkinWeights(USkeletalMesh* SkeletalMeshDest,
 	TArray<FText> WarningMessages;
 	TArray<FName> WarningNames;
 	//Build the destination mesh with the Alternate influences, so the chunking is done properly.
-	bool bBuildSuccess = MeshUtilities.BuildSkeletalMesh(LODModelDest, SkeletalMeshDest->RefSkeleton, LODInfluencesDest, LODWedgesDest, LODFacesDest, LODPointsDest, LODPointToRawMapDest, BuildOptions, &WarningMessages, &WarningNames);
+	bool bBuildSuccess = MeshUtilities.BuildSkeletalMesh(LODModelDest, RefSkeleton, LODInfluencesDest, LODWedgesDest, LODFacesDest, LODPointsDest, LODPointToRawMapDest, BuildOptions, &WarningMessages, &WarningNames);
+	//Re-Apply the user section changes, the UserSectionsData is map to original section and should match the builded LODModel
+	LODModelDest.SyncronizeUserSectionsDataArray();
+
 	RegenerateAllImportSkinWeightProfileData(LODModelDest);
 	
 	return bBuildSuccess;
@@ -2004,7 +1947,8 @@ void FLODUtilities::RegenerateDependentLODs(USkeletalMesh* SkeletalMesh, int32 L
 {
 	FSkeletalMeshUpdateContext UpdateContext;
 	UpdateContext.SkeletalMesh = SkeletalMesh;
-	//Check the dependencies and regenerate the LODs acoording to it
+
+	//Check the dependencies and regenerate the LODs according to it
 	TArray<bool> LODDependencies;
 	int32 LODNumber = SkeletalMesh->GetLODNum();
 	LODDependencies.AddZeroed(LODNumber);
@@ -2022,19 +1966,20 @@ void FLODUtilities::RegenerateDependentLODs(USkeletalMesh* SkeletalMesh, int32 L
 	}
 	if (bRegenLODs)
 	{
-		TComponentReregisterContext<USkinnedMeshComponent> ReregisterContext;
-		SkeletalMesh->Modify();
-		SkeletalMesh->ReleaseResources();
-		SkeletalMesh->ReleaseResourcesFence.Wait();
-		for (int32 DependentLODIndex = LODIndex + 1; DependentLODIndex < LODNumber; ++DependentLODIndex)
+		FScopedSkeletalMeshPostEditChange ScopedPostEditChange(SkeletalMesh);
+		auto SimplifyDependentLODS = [&LODIndex, &LODNumber, &LODDependencies, &UpdateContext]()
 		{
-			if (LODDependencies[DependentLODIndex])
+			for (int32 DependentLODIndex = LODIndex + 1; DependentLODIndex < LODNumber; ++DependentLODIndex)
 			{
-				FLODUtilities::SimplifySkeletalMeshLOD(UpdateContext, DependentLODIndex, false);
+				if (LODDependencies[DependentLODIndex])
+				{
+					FLODUtilities::SimplifySkeletalMeshLOD(UpdateContext, DependentLODIndex, true);
+				}
 			}
-		}
-		SkeletalMesh->PostEditChange();
-		SkeletalMesh->InitResources();
+		};
+
+		SkeletalMesh->Modify();
+		SimplifyDependentLODS();
 	}
 }
 
@@ -2486,4 +2431,84 @@ void FLODUtilities::BuildMorphTargets(USkeletalMesh* BaseSkelMesh, FSkeletalMesh
 	}
 }
 
+void FLODUtilities::UnbindClothingAndBackup(USkeletalMesh* SkeletalMesh, TArray<ClothingAssetUtils::FClothingAssetMeshBinding>& ClothingBindings)
+{
+	for (int32 LODIndex = 0; LODIndex < SkeletalMesh->GetImportedModel()->LODModels.Num(); ++LODIndex)
+	{
+		TArray<ClothingAssetUtils::FClothingAssetMeshBinding> LODBindings;
+		UnbindClothingAndBackup(SkeletalMesh, LODBindings, LODIndex);
+		ClothingBindings.Append(LODBindings);
+	}
+}
+
+void FLODUtilities::UnbindClothingAndBackup(USkeletalMesh* SkeletalMesh, TArray<ClothingAssetUtils::FClothingAssetMeshBinding>& ClothingBindings, const int32 LODIndex)
+{
+	if (!SkeletalMesh->GetImportedModel()->LODModels.IsValidIndex(LODIndex))
+	{
+		return;
+	}
+	FSkeletalMeshLODModel& LODModel = SkeletalMesh->GetImportedModel()->LODModels[LODIndex];
+	//Store the clothBinding
+	ClothingAssetUtils::GetMeshClothingAssetBindings(SkeletalMesh, ClothingBindings, LODIndex);
+	//Unbind the Cloth for this LOD before we reduce it, we will put back the cloth after the reduction, if it still match the sections
+	for (ClothingAssetUtils::FClothingAssetMeshBinding& Binding : ClothingBindings)
+	{
+		if (Binding.LODIndex == LODIndex)
+		{
+			check(Binding.Asset);
+			Binding.Asset->UnbindFromSkeletalMesh(SkeletalMesh, Binding.LODIndex);
+			
+			//Use the UserSectionsData original section index, this will ensure we remap correctly the cloth if the reduction has change the number of sections
+			int32 OriginalDataSectionIndex = LODModel.Sections[Binding.SectionIndex].OriginalDataSectionIndex;
+			Binding.SectionIndex = OriginalDataSectionIndex;
+			
+			FSkelMeshSourceSectionUserData& SectionUserData = LODModel.UserSectionsData.FindChecked(OriginalDataSectionIndex);
+			SectionUserData.ClothingData.AssetGuid = FGuid();
+			SectionUserData.ClothingData.AssetLodIndex = INDEX_NONE;
+			SectionUserData.CorrespondClothAssetIndex = INDEX_NONE;
+		}
+	}
+}
+
+void FLODUtilities::RestoreClothingFromBackup(USkeletalMesh* SkeletalMesh, TArray<ClothingAssetUtils::FClothingAssetMeshBinding>& ClothingBindings)
+{
+	for (int32 LODIndex = 0; LODIndex < SkeletalMesh->GetImportedModel()->LODModels.Num(); ++LODIndex)
+	{
+		RestoreClothingFromBackup(SkeletalMesh, ClothingBindings, LODIndex);
+	}
+}
+
+void FLODUtilities::RestoreClothingFromBackup(USkeletalMesh* SkeletalMesh, TArray<ClothingAssetUtils::FClothingAssetMeshBinding>& ClothingBindings, const int32 LODIndex)
+{
+	if (!SkeletalMesh->GetImportedModel()->LODModels.IsValidIndex(LODIndex))
+	{
+		return;
+	}
+	FSkeletalMeshLODModel& LODModel = SkeletalMesh->GetImportedModel()->LODModels[LODIndex];
+	for (ClothingAssetUtils::FClothingAssetMeshBinding& Binding : ClothingBindings)
+	{
+		for (int32 SectionIndex = 0; SectionIndex < LODModel.Sections.Num(); ++SectionIndex)
+		{
+			if (LODModel.Sections[SectionIndex].OriginalDataSectionIndex != Binding.SectionIndex)
+			{
+				continue;
+			}
+			if (Binding.LODIndex == LODIndex)
+			{
+				check(Binding.Asset);
+				if (Binding.Asset->BindToSkeletalMesh(SkeletalMesh, Binding.LODIndex, SectionIndex, Binding.AssetInternalLodIndex))
+				{
+					//If successfull set back the section user data
+					FSkelMeshSourceSectionUserData& SectionUserData = LODModel.UserSectionsData.FindChecked(Binding.SectionIndex);
+					SectionUserData.CorrespondClothAssetIndex = LODModel.Sections[SectionIndex].CorrespondClothAssetIndex;
+					SectionUserData.ClothingData = LODModel.Sections[SectionIndex].ClothingData;
+				}
+			}
+			break;
+		}
+	}
+}
+
+
+
 #undef LOCTEXT_NAMESPACE // "LODUtilities"
\ No newline at end of file
diff --git a/Engine/Source/Developer/SkeletalMeshUtilitiesCommon/Public/LODUtilities.h b/Engine/Source/Developer/SkeletalMeshUtilitiesCommon/Public/LODUtilities.h
--- a/Engine/Source/Developer/SkeletalMeshUtilitiesCommon/Public/LODUtilities.h
+++ b/Engine/Source/Developer/SkeletalMeshUtilitiesCommon/Public/LODUtilities.h
@@ -8,7 +8,10 @@
 #include "Rendering/SkeletalMeshLODImporterData.h"
 #include "Framework/Commands/UIAction.h"
 
-
+namespace ClothingAssetUtils
+{
+	struct FClothingAssetMeshBinding;
+}
 
 //////////////////////////////////////////////////////////////////////////
 // FSkeletalMeshUpdateContext
@@ -25,6 +28,7 @@ struct FSkeletalMeshUpdateContext
 //////////////////////////////////////////////////////////////////////////
 // FLODUtilities
 
+
 class SKELETALMESHUTILITIESCOMMON_API FLODUtilities
 {
 public:
@@ -62,7 +66,7 @@ public:
 	* @param UpdateContext - The skeletal mesh and actor components to operate on.
 	* @param DesiredLOD - The LOD to simplify
 	*/
-	static void SimplifySkeletalMeshLOD(FSkeletalMeshUpdateContext& UpdateContext, int32 DesiredLOD, bool bReregisterComponent = true, bool bRestoreClothing = false);
+	static void SimplifySkeletalMeshLOD(FSkeletalMeshUpdateContext& UpdateContext, int32 DesiredLOD, bool bRestoreClothing = false);
 
 	/**
 	*	Restore the LOD imported model to the last imported data. Call this function if you want to remove the reduce on the base LOD
@@ -71,7 +75,7 @@ public:
 	* @param LodIndex - The LOD index to restore the imported LOD model
 	* @param bReregisterComponent - if true the component using the skeletal mesh will all be re register.
 	*/
-	static void RestoreSkeletalMeshLODImportedData(USkeletalMesh* SkeletalMesh, int32 LodIndex, bool bReregisterComponent = true);
+	static void RestoreSkeletalMeshLODImportedData(USkeletalMesh* SkeletalMesh, int32 LodIndex);
 	
 	/**
 	 * Refresh LOD Change
@@ -117,9 +121,18 @@ public:
 	 */
 	static bool UpdateAlternateSkinWeights(USkeletalMesh* SkeletalMeshDest, const FName& ProfileNameDest, int32 LODIndexDest, FOverlappingThresholds OverlappingThresholds, bool ShouldImportNormals, bool ShouldImportTangents, bool bUseMikkTSpace, bool bComputeWeightedNormals);
 
+	static bool UpdateAlternateSkinWeights(FSkeletalMeshLODModel& LODModelDest, const FString SkeletalMeshName, FReferenceSkeleton& RefSkeleton, const FName& ProfileNameDest, int32 LODIndexDest, FOverlappingThresholds OverlappingThresholds, bool ShouldImportNormals, bool ShouldImportTangents, bool bUseMikkTSpace, bool bComputeWeightedNormals);
+
 	/** Re-generate all (editor-only) skin weight profile, used whenever we rebuild the skeletal mesh data which could change the chunking and bone indices */
 	static void RegenerateAllImportSkinWeightProfileData(FSkeletalMeshLODModel& LODModelDest);
 
+	static void UnbindClothingAndBackup(USkeletalMesh* SkeletalMesh, TArray<ClothingAssetUtils::FClothingAssetMeshBinding>& ClothingBindings);
+	static void UnbindClothingAndBackup(USkeletalMesh* SkeletalMesh, TArray<ClothingAssetUtils::FClothingAssetMeshBinding>& ClothingBindings, const int32 LODIndex);
+	
+	static void RestoreClothingFromBackup(USkeletalMesh* SkeletalMesh, TArray<ClothingAssetUtils::FClothingAssetMeshBinding>& ClothingBindings);
+	static void RestoreClothingFromBackup(USkeletalMesh* SkeletalMesh, TArray<ClothingAssetUtils::FClothingAssetMeshBinding>& ClothingBindings, const int32 LODIndex);
+	
+
 private:
 	FLODUtilities() {}
 
@@ -134,7 +147,7 @@ private:
 	 * @param SkeletalMesh - The skeletal mesh and actor components to operate on.
 	 * @param DesiredLOD - Desired LOD
 	 */
-	static void SimplifySkeletalMeshLOD(USkeletalMesh* SkeletalMesh, int32 DesiredLOD, bool bReregisterComponent = true, bool bRestoreClothing = false);
+	static void SimplifySkeletalMeshLOD(USkeletalMesh* SkeletalMesh, int32 DesiredLOD, bool bRestoreClothing = false);
 
 	/**
 	*  Remap the morph targets of the base LOD onto the desired LOD.
diff --git a/Engine/Source/Editor/Experimental/EditorInteractiveToolsFramework/EditorInteractiveToolsFramework.Build.cs b/Engine/Source/Editor/Experimental/EditorInteractiveToolsFramework/EditorInteractiveToolsFramework.Build.cs
--- a/Engine/Source/Editor/Experimental/EditorInteractiveToolsFramework/EditorInteractiveToolsFramework.Build.cs
+++ b/Engine/Source/Editor/Experimental/EditorInteractiveToolsFramework/EditorInteractiveToolsFramework.Build.cs
@@ -45,7 +45,9 @@ public class EditorInteractiveToolsFramework : ModuleRules
                 "LevelEditor",
                 "ApplicationCore",
                 "EditorStyle",
-                "InteractiveToolsFramework"
+                "InteractiveToolsFramework",
+				"MeshDescription",
+				"StaticMeshDescription"
 
 				// ... add private dependencies that you statically link with here ...	
 			}
diff --git a/Engine/Source/Editor/MeshPaint/Private/MeshPaintHelpers.cpp b/Engine/Source/Editor/MeshPaint/Private/MeshPaintHelpers.cpp
--- a/Engine/Source/Editor/MeshPaint/Private/MeshPaintHelpers.cpp
+++ b/Engine/Source/Editor/MeshPaint/Private/MeshPaintHelpers.cpp
@@ -14,10 +14,7 @@
 #include "Engine/SkeletalMesh.h"
 #include "Engine/Texture2D.h"
 #include "StaticMeshResources.h"
-
-#include "MeshDescription.h"
-#include "MeshAttributes.h"
-#include "MeshAttributeArray.h"
+#include "StaticMeshAttributes.h"
 
 #include "Rendering/SkeletalMeshRenderData.h"
 
@@ -38,7 +35,7 @@
 
 #include "Editor.h"
 #include "LevelEditor.h"
-#include "ILevelViewport.h"
+#include "IAssetViewport.h"
 #include "EditorViewportClient.h"
 #include "LevelEditorViewport.h"
 
@@ -112,22 +109,21 @@ bool MeshPaintHelpers::PropagateColorsToRawMesh(UStaticMesh* StaticMesh, int32 L
 	if (RenderData.WedgeMap.Num() > 0 && ColorVertexBuffer.GetNumVertices() == RenderModel.GetNumVertices())
 	{
 		// Use the wedge map if it is available as it is lossless.
-		FMeshDescription* RawMeshPtr = StaticMesh->GetMeshDescription(LODIndex);
+		FMeshDescription* MeshDescription = StaticMesh->GetMeshDescription(LODIndex);
 
-		if (RawMeshPtr == nullptr)
+		if (MeshDescription == nullptr)
 		{
 			//Cannot propagate to a generated LOD, the generated LOD use the source LOD vertex painting. 
 			return false;
 		}
 
-		FMeshDescription& RawMesh = *RawMeshPtr;
-		int32 NumWedges = RawMesh.VertexInstances().Num();
+		FStaticMeshAttributes Attributes(*MeshDescription);
+		int32 NumWedges = MeshDescription->VertexInstances().Num();
 		if (RenderData.WedgeMap.Num() == NumWedges)
 		{
-			FStaticMeshDescriptionAttributeGetter AttributeGetter(&RawMesh);
-			TVertexInstanceAttributesRef<FVector4> Colors = AttributeGetter.GetColors();
+			TVertexInstanceAttributesRef<FVector4> Colors = Attributes.GetVertexInstanceColors();
 			int32 VertexInstanceIndex = 0;
-			for (const FVertexInstanceID VertexInstanceID : RawMesh.VertexInstances().GetElementIDs())
+			for (const FVertexInstanceID VertexInstanceID : MeshDescription->VertexInstances().GetElementIDs())
 			{
 				FLinearColor WedgeColor = FLinearColor::White;
 				int32 Index = RenderData.WedgeMap[VertexInstanceIndex];
@@ -144,24 +140,25 @@ bool MeshPaintHelpers::PropagateColorsToRawMesh(UStaticMesh* StaticMesh, int32 L
 	}
 	else
 	{
-		FMeshDescription* SrcModelMeshDescription = StaticMesh->GetMeshDescription(LODIndex);
+		FMeshDescription* MeshDescription = StaticMesh->GetMeshDescription(LODIndex);
 		// If there's no raw mesh data, don't try to do any fixup here
-		if (SrcModelMeshDescription == nullptr || ComponentLODInfo.OverrideVertexColors == nullptr)
+		if (MeshDescription == nullptr || ComponentLODInfo.OverrideVertexColors == nullptr)
 		{
 			return false;
 		}
 
+		FStaticMeshAttributes Attributes(*MeshDescription);
+
 		// Fall back to mapping based on position.
-		FStaticMeshDescriptionAttributeGetter AttributeGetter(SrcModelMeshDescription);
-		TVertexAttributesConstRef<FVector> VertexPositions = AttributeGetter.GetPositionsConst();
-		TVertexInstanceAttributesRef<FVector4> Colors = AttributeGetter.GetColors();
+		TVertexAttributesConstRef<FVector> VertexPositions = Attributes.GetVertexPositions();
+		TVertexInstanceAttributesRef<FVector4> Colors = Attributes.GetVertexInstanceColors();
 		TArray<FColor> NewVertexColors;
 		FPositionVertexBuffer TempPositionVertexBuffer;
-		int32 NumVertex = SrcModelMeshDescription->Vertices().Num();
+		int32 NumVertex = MeshDescription->Vertices().Num();
 		TArray<FVector> VertexPositionsDup;
 		VertexPositionsDup.AddZeroed(NumVertex);
 		int32 VertexIndex = 0;
-		for (const FVertexID VertexID : SrcModelMeshDescription->Vertices().GetElementIDs())
+		for (const FVertexID VertexID : MeshDescription->Vertices().GetElementIDs())
 		{
 			VertexPositionsDup[VertexIndex] = VertexPositions[VertexID];
 		}
@@ -177,9 +174,9 @@ bool MeshPaintHelpers::PropagateColorsToRawMesh(UStaticMesh* StaticMesh, int32 L
 		);
 		if (NewVertexColors.Num() == NumVertex)
 		{
-			for (const FVertexInstanceID VertexInstanceID : SrcModelMeshDescription->VertexInstances().GetElementIDs())
+			for (const FVertexInstanceID VertexInstanceID : MeshDescription->VertexInstances().GetElementIDs())
 			{
-				const FVertexID VertexID = SrcModelMeshDescription->GetVertexInstanceVertex(VertexInstanceID);
+				const FVertexID VertexID = MeshDescription->GetVertexInstanceVertex(VertexInstanceID);
 				Colors[VertexInstanceID] = FLinearColor(NewVertexColors[VertexID.GetValue()]);
 			}
 			StaticMesh->CommitMeshDescription(LODIndex);
@@ -1152,11 +1149,11 @@ void MeshPaintHelpers::SetViewportColorMode(EMeshPaintColorViewMode ColorViewMod
 void MeshPaintHelpers::SetRealtimeViewport(bool bRealtime)
 {
 	FLevelEditorModule& LevelEditorModule = FModuleManager::GetModuleChecked<FLevelEditorModule>("LevelEditor");
-	TSharedPtr< ILevelViewport > ViewportWindow = LevelEditorModule.GetFirstActiveViewport();
+	TSharedPtr< IAssetViewport > ViewportWindow = LevelEditorModule.GetFirstActiveViewport();
 	const bool bRememberCurrentState = false;
 	if (ViewportWindow.IsValid())
 	{
-		FEditorViewportClient &Viewport = ViewportWindow->GetLevelViewportClient();
+		FEditorViewportClient &Viewport = ViewportWindow->GetAssetViewportClient();
 		if (Viewport.IsPerspective())
 		{
 			if (bRealtime)
diff --git a/Engine/Source/Editor/Persona/Private/AnimGraphNodeDetails.cpp b/Engine/Source/Editor/Persona/Private/AnimGraphNodeDetails.cpp
--- a/Engine/Source/Editor/Persona/Private/AnimGraphNodeDetails.cpp
+++ b/Engine/Source/Editor/Persona/Private/AnimGraphNodeDetails.cpp
@@ -44,6 +44,8 @@
 #include "Engine/SkeletalMeshSocket.h"
 #include "Styling/CoreStyle.h"
 #include "Widgets/Input/SCheckBox.h"
+#include "LODInfoUILayout.h"
+#include "IPersonaToolkit.h"
 
 #define LOCTEXT_NAMESPACE "KismetNodeWithOptionalPinsDetails"
 
@@ -516,6 +518,7 @@ void FBoneReferenceCustomization::SetEditableSkeleton(TSharedRef<IPropertyHandle
 	UAnimGraphNode_Base* AnimGraphNode = nullptr;
 	USkeletalMesh* SkeletalMesh = nullptr;
 	UAnimationAsset * AnimationAsset = nullptr;
+	ULODInfoUILayout* LODInfoUILayout = nullptr;
 
 	USkeleton* TargetSkeleton = nullptr;
 	TSharedPtr<IEditableSkeleton> EditableSkeleton;
@@ -534,6 +537,16 @@ void FBoneReferenceCustomization::SetEditableSkeleton(TSharedRef<IPropertyHandle
 			TargetSkeleton = SkeletalMesh->Skeleton;
 			break;
 		}
+
+		LODInfoUILayout = Cast<ULODInfoUILayout>(*OuterIter);
+		if (LODInfoUILayout)
+		{
+			SkeletalMesh = LODInfoUILayout->GetPersonaToolkit()->GetPreviewMesh();
+			check(SkeletalMesh);
+			TargetSkeleton = SkeletalMesh->Skeleton;
+			break;
+		}
+
 		AnimationAsset = Cast<UAnimationAsset>(*OuterIter);
 		if (AnimationAsset)
 		{
