Fix crash when re-importing alternate skinning. This is a fallback from the skeletalmesh refactor

diff --git a/Engine/Source/Developer/MeshUtilities/Private/MeshUtilities.cpp b/Engine/Source/Developer/MeshUtilities/Private/MeshUtilities.cpp
--- a/Engine/Source/Developer/MeshUtilities/Private/MeshUtilities.cpp
+++ b/Engine/Source/Developer/MeshUtilities/Private/MeshUtilities.cpp
@@ -3501,7 +3501,7 @@ public:
 		
 		int32 NumFaces = BuildData->GetNumFaces();
 		int32 NumWedges = BuildData->GetNumWedges();
-		check(NumFaces * 3 == NumWedges);
+		check(NumFaces * 3 <= NumWedges);
 		
 		// Compute per-triangle tangents.
 		TArray<FVector> TriangleTangentX;
@@ -3615,12 +3615,15 @@ public:
 				for (int32 k = 0; k < DupVerts.Num(); k++)
 				{
 					int32 PotentialTriangleIndex = DupVerts[k] / 3;
-					
-					bool bDegeneratedTriangles = TriangleTangentZ[FaceIndex].IsNearlyZero() || TriangleTangentZ[PotentialTriangleIndex].IsNearlyZero();
-					//Do not add mirror triangle to the adjacentFaces. Also make sure adjacent triangle is in the same connected triangle patch. Accept connected degenerate triangle
-					if ((bDegeneratedTriangles || !IsTriangleMirror(BuildData, TriangleTangentZ, FaceIndex, PotentialTriangleIndex)) && PatchIndex == FaceIndexToPatchIndex[PotentialTriangleIndex])
+					//Do not add a triangle that was remove
+					if (TriangleTangentZ.IsValidIndex(PotentialTriangleIndex))
 					{
-						AdjacentFaces.AddUnique(PotentialTriangleIndex);
+						bool bDegeneratedTriangles = TriangleTangentZ[FaceIndex].IsNearlyZero() || TriangleTangentZ[PotentialTriangleIndex].IsNearlyZero();
+						//Do not add mirror triangle to the adjacentFaces. Also make sure adjacent triangle is in the same connected triangle patch. Accept connected degenerate triangle
+						if ((bDegeneratedTriangles || !IsTriangleMirror(BuildData, TriangleTangentZ, FaceIndex, PotentialTriangleIndex)) && PatchIndex == FaceIndexToPatchIndex[PotentialTriangleIndex])
+						{
+							AdjacentFaces.AddUnique(PotentialTriangleIndex);
+						}
 					}
 				}
 			}
