Fix build linker error FStaticMeshAttributes::FStaticMeshAttributes
#rnx

diff --git a/Engine/Source/Developer/MeshUtilities/MeshUtilities.Build.cs b/Engine/Source/Developer/MeshUtilities/MeshUtilities.Build.cs
--- a/Engine/Source/Developer/MeshUtilities/MeshUtilities.Build.cs
+++ b/Engine/Source/Developer/MeshUtilities/MeshUtilities.Build.cs
@@ -36,6 +36,7 @@ public class MeshUtilities : ModuleRules
                 "MeshBuilderCommon",
                 "MeshUtilitiesCommon",
                 "MeshDescription",
+				"StaticMeshDescription",
                 "MeshDescriptionOperations",
 				"ToolMenus",
             }
