Fix skeletalmesh conversion to use the LODMaterialMap when converting to a staticmesh.

diff --git a/Engine/Source/Developer/MeshUtilities/Private/MeshUtilities.cpp b/Engine/Source/Developer/MeshUtilities/Private/MeshUtilities.cpp
--- a/Engine/Source/Developer/MeshUtilities/Private/MeshUtilities.cpp
+++ b/Engine/Source/Developer/MeshUtilities/Private/MeshUtilities.cpp
@@ -303,10 +303,10 @@ static void SkinnedMeshToRawMeshes(USkinnedMeshComponent* InSkinnedMeshComponent
 				}
 
 				int32 MaterialIndex = SkelMeshSection.MaterialIndex;
-				// use the remapping of material indices for all LODs besides the base LOD 
-				if (LODIndexRead > 0 && SrcLODInfo.LODMaterialMap.IsValidIndex(SkelMeshSection.MaterialIndex))
+				// use the remapping of material indices if there is a valid value
+				if (SrcLODInfo.LODMaterialMap.IsValidIndex(SectionIndex) && SrcLODInfo.LODMaterialMap[SectionIndex] != INDEX_NONE)
 				{
-					MaterialIndex = FMath::Clamp<int32>(SrcLODInfo.LODMaterialMap[SkelMeshSection.MaterialIndex], 0, InSkinnedMeshComponent->SkeletalMesh->Materials.Num());
+					MaterialIndex = FMath::Clamp<int32>(SrcLODInfo.LODMaterialMap[SectionIndex], 0, InSkinnedMeshComponent->SkeletalMesh->Materials.Num());
 				}
 
 				// copy face info
