Fix conversion from SkeletalMesh to FRawMesh.
Bi-Normal vector should be compute like (Z^X) * biNorm Sign.

diff --git a/Engine/Source/Developer/MeshUtilities/Private/MeshUtilities.cpp b/Engine/Source/Developer/MeshUtilities/Private/MeshUtilities.cpp
--- a/Engine/Source/Developer/MeshUtilities/Private/MeshUtilities.cpp
+++ b/Engine/Source/Developer/MeshUtilities/Private/MeshUtilities.cpp
@@ -266,7 +266,7 @@ static void SkinnedMeshToRawMeshes(USkinnedMeshComponent* InSkinnedMeshComponent
 					const FVector TangentX = InComponentToWorld.TransformVector(SkinnedVertex.TangentX.ToFVector());
 					const FVector TangentZ = InComponentToWorld.TransformVector(SkinnedVertex.TangentZ.ToFVector());
 					const FVector4 UnpackedTangentZ = SkinnedVertex.TangentZ.ToFVector4();
-					const FVector TangentY = (TangentX ^ TangentZ).GetSafeNormal() * UnpackedTangentZ.W;
+					const FVector TangentY = (TangentZ ^ TangentX).GetSafeNormal() * UnpackedTangentZ.W;
 
 					RawMesh.WedgeTangentX.Add(TangentX);
 					RawMesh.WedgeTangentY.Add(TangentY);
