Fix skeletalmesh temporary build structure to support high precision tangents.

diff --git a/Engine/Source/Developer/MeshUtilities/Private/MeshUtilities.cpp b/Engine/Source/Developer/MeshUtilities/Private/MeshUtilities.cpp
--- a/Engine/Source/Developer/MeshUtilities/Private/MeshUtilities.cpp
+++ b/Engine/Source/Developer/MeshUtilities/Private/MeshUtilities.cpp
@@ -747,9 +747,9 @@ void FMeshUtilities::BuildSkeletalModelFromChunks(FSkeletalMeshLODModel& LODMode
 
 			FSoftSkinVertex NewVertex;
 			NewVertex.Position = SoftVertex.Position;
-			NewVertex.TangentX = SoftVertex.TangentX.ToFVector();
-			NewVertex.TangentY = SoftVertex.TangentY.ToFVector();
-			NewVertex.TangentZ = SoftVertex.TangentZ.ToFVector();
+			NewVertex.TangentX = SoftVertex.TangentX;
+			NewVertex.TangentY = SoftVertex.TangentY;
+			NewVertex.TangentZ = SoftVertex.TangentZ;
 			FMemory::Memcpy(NewVertex.UVs, SoftVertex.UVs, sizeof(FVector2D)*MAX_TEXCOORDS);
 			NewVertex.Color = SoftVertex.Color;
 			for (int32 i = 0; i < MAX_TOTAL_INFLUENCES; ++i)
diff --git a/Engine/Source/Developer/MeshUtilities/Private/SkeletalMeshTools.cpp b/Engine/Source/Developer/MeshUtilities/Private/SkeletalMeshTools.cpp
--- a/Engine/Source/Developer/MeshUtilities/Private/SkeletalMeshTools.cpp
+++ b/Engine/Source/Developer/MeshUtilities/Private/SkeletalMeshTools.cpp
@@ -24,17 +24,17 @@ namespace SkeletalMeshTools
 			}
 		}
 
-		if(!NormalsEqual(V1.TangentX.ToFVector(), V2.TangentX.ToFVector(), OverlappingThresholds))
+		if(!NormalsEqual(V1.TangentX, V2.TangentX, OverlappingThresholds))
 		{
 			return false;
 		}
 
-		if(!NormalsEqual(V1.TangentY.ToFVector(), V2.TangentY.ToFVector(), OverlappingThresholds))
+		if(!NormalsEqual(V1.TangentY, V2.TangentY, OverlappingThresholds))
 		{
 			return false;
 		}
 
-		if(!NormalsEqual(V1.TangentZ.ToFVector(), V2.TangentZ.ToFVector(), OverlappingThresholds))
+		if(!NormalsEqual(V1.TangentZ, V2.TangentZ, OverlappingThresholds))
 		{
 			return false;
 		}
diff --git a/Engine/Source/Developer/MeshUtilities/Private/SkeletalMeshTools.h b/Engine/Source/Developer/MeshUtilities/Private/SkeletalMeshTools.h
--- a/Engine/Source/Developer/MeshUtilities/Private/SkeletalMeshTools.h
+++ b/Engine/Source/Developer/MeshUtilities/Private/SkeletalMeshTools.h
@@ -4,7 +4,6 @@
 
 #include "CoreMinimal.h"
 #include "Containers/IndirectArray.h"
-#include "PackedNormal.h"
 #include "GPUSkinPublicDefs.h"
 #include "Components.h"
 #include "BoneIndices.h"
@@ -26,7 +25,7 @@ struct FSkeletalMeshVertIndexAndZ
 struct FSoftSkinBuildVertex
 {
 	FVector			Position;
-	FPackedNormal	TangentX,	// Tangent, U-direction
+	FVector			TangentX,	// Tangent, U-direction
 					TangentY,	// Binormal, V-direction
 					TangentZ;	// Normal
 	FVector2D		UVs[MAX_TEXCOORDS]; // UVs
