When converting an actor to a Static Mesh, it is now focused in the Content Browser. This solves the bug detailed in UE-76250 about the new Mesh not being focused.

diff --git a/Engine/Source/Developer/MeshUtilities/Private/MeshUtilities.cpp b/Engine/Source/Developer/MeshUtilities/Private/MeshUtilities.cpp
--- a/Engine/Source/Developer/MeshUtilities/Private/MeshUtilities.cpp
+++ b/Engine/Source/Developer/MeshUtilities/Private/MeshUtilities.cpp
@@ -43,7 +43,6 @@
 #include "HierarchicalLODUtilitiesModule.h"
 #include "MeshBoneReduction.h"
 #include "MeshMergeData.h"
-#include "Editor/EditorPerProjectUserSettings.h"
 #include "GPUSkinVertexFactory.h"
 #include "Developer/AssetTools/Public/IAssetTools.h"
 #include "Developer/AssetTools/Public/AssetToolsModule.h"
@@ -68,6 +67,8 @@
 #include "IAnimationBlueprintEditorModule.h"
 #include "IAnimationEditor.h"
 #include "IAnimationEditorModule.h"
+#include "IContentBrowserSingleton.h"
+#include "ContentBrowserModule.h"
 #include "ISkeletalMeshEditor.h"
 #include "ISkeletalMeshEditorModule.h"
 #include "ISkeletonEditor.h"
@@ -82,6 +83,7 @@
 #include "Engine/SkeletalMeshSimplificationSettings.h"
 #include "Engine/ProxyLODMeshSimplificationSettings.h"
 
+#include "Editor/EditorPerProjectUserSettings.h"
 #include "IDetailCustomization.h"
 #include "EditorStyleSet.h"
 #include "PropertyEditorModule.h"
@@ -400,6 +402,8 @@ static void StaticMeshToRawMeshes(UStaticMeshComponent* InStaticMeshComponent, i
 
 UStaticMesh* FMeshUtilities::ConvertMeshesToStaticMesh(const TArray<UMeshComponent*>& InMeshComponents, const FTransform& InRootTransform, const FString& InPackageName)
 {
+	UStaticMesh* StaticMesh = nullptr;
+
 	// Build a package name to use
 	FString MeshName;
 	FString PackageName;
@@ -528,7 +532,7 @@ UStaticMesh* FMeshUtilities::ConvertMeshesToStaticMesh(const TArray<UMeshCompone
 			check(Package);
 
 			// Create StaticMesh object
-			UStaticMesh* StaticMesh = NewObject<UStaticMesh>(Package, *MeshName, RF_Public | RF_Standalone);
+			StaticMesh = NewObject<UStaticMesh>(Package, *MeshName, RF_Public | RF_Standalone);
 			StaticMesh->InitResources();
 
 			StaticMesh->LightingGuid = FGuid::NewGuid();
@@ -611,7 +615,7 @@ UStaticMesh* FMeshUtilities::ConvertMeshesToStaticMesh(const TArray<UMeshCompone
 		}
 	}
 
-	return nullptr;
+	return StaticMesh;
 }
 
 /**
@@ -6151,7 +6155,7 @@ TSharedRef<FExtender> FMeshUtilities::GetLevelViewportContextMenuExtender(const
 					FText::Format(LOCTEXT("ConvertSelectedActorsToStaticMeshText", "Convert {0} To Static Mesh"), ActorName),
 					LOCTEXT("ConvertSelectedActorsToStaticMeshTooltip", "Convert the selected actor's meshes to a new Static Mesh asset. Supports static and skeletal meshes."),
 					FSlateIcon(),
-					FUIAction(FExecuteAction::CreateRaw(this, &FMeshUtilities::ConvertActorMeshesToStaticMesh, InActors))
+					FUIAction(FExecuteAction::CreateRaw(this, &FMeshUtilities::ConvertActorMeshesToStaticMeshUIAction, InActors))
 				);
 			})
 			);
@@ -6161,7 +6165,7 @@ TSharedRef<FExtender> FMeshUtilities::GetLevelViewportContextMenuExtender(const
 	return Extender;
 }
 
-void FMeshUtilities::ConvertActorMeshesToStaticMesh(const TArray<AActor*> InActors)
+void FMeshUtilities::ConvertActorMeshesToStaticMeshUIAction(const TArray<AActor*> InActors)
 {
 	TArray<UMeshComponent*> MeshComponents;
 
@@ -6207,7 +6211,14 @@ void FMeshUtilities::ConvertActorMeshesToStaticMesh(const TArray<AActor*> InActo
 		RootTransform.SetLocation(Location);
 	}
 
-	ConvertMeshesToStaticMesh(MeshComponents, RootTransform);
+	UStaticMesh* StaticMesh = ConvertMeshesToStaticMesh(MeshComponents, RootTransform);
+
+	// Also notify the content browser that the new assets exists
+	if (StaticMesh != nullptr)
+	{
+		FContentBrowserModule& ContentBrowserModule = FModuleManager::Get().LoadModuleChecked<FContentBrowserModule>("ContentBrowser");
+		ContentBrowserModule.Get().SyncBrowserToAssets(TArray<UObject*>({ StaticMesh }), true);
+	}
 }
 
 /************************************************************************/
diff --git a/Engine/Source/Developer/MeshUtilities/Private/MeshUtilitiesPrivate.h b/Engine/Source/Developer/MeshUtilities/Private/MeshUtilitiesPrivate.h
--- a/Engine/Source/Developer/MeshUtilities/Private/MeshUtilitiesPrivate.h
+++ b/Engine/Source/Developer/MeshUtilities/Private/MeshUtilitiesPrivate.h
@@ -200,7 +200,7 @@ protected:
 
 	TSharedRef<FExtender> GetLevelViewportContextMenuExtender(const TSharedRef<FUICommandList> CommandList, const TArray<AActor*> InActors);
 
-	void ConvertActorMeshesToStaticMesh(const TArray<AActor*> InActors);
+	void ConvertActorMeshesToStaticMeshUIAction(const TArray<AActor*> InActors);
 
 	FDelegateHandle ModuleLoadedDelegateHandle;
 	FDelegateHandle LevelViewportExtenderHandle;
