Fix debug editor crash on Linux due to Persona static init to early during loading

diff --git a/Engine/Source/Developer/MeshUtilities/MeshUtilities.Build.cs b/Engine/Source/Developer/MeshUtilities/MeshUtilities.Build.cs
--- a/Engine/Source/Developer/MeshUtilities/MeshUtilities.Build.cs
+++ b/Engine/Source/Developer/MeshUtilities/MeshUtilities.Build.cs
@@ -30,10 +30,6 @@ public class MeshUtilities : ModuleRules
 				"HierarchicalLODUtilities",
 				"Landscape",
 				"LevelEditor",
-				"AnimationBlueprintEditor",
-				"AnimationEditor",
-				"SkeletalMeshEditor",
-				"SkeletonEditor",
 				"PropertyEditor",
 				"EditorStyle",
                 "GraphColor",
@@ -52,15 +48,22 @@ public class MeshUtilities : ModuleRules
 
         PrivateIncludePathModuleNames.AddRange(
           new string[] {
+				"AnimationBlueprintEditor",
+				"AnimationEditor",
                 "MeshMergeUtilities",
                 "MaterialBaking",
+				"Persona",
+				"SkeletalMeshEditor",
           }
       );
 
         DynamicallyLoadedModuleNames.AddRange(
             new string[] {
+				"AnimationBlueprintEditor",
+				"AnimationEditor",
                 "MeshMergeUtilities",
                 "MaterialBaking",
+				"SkeletalMeshEditor",
             }
         );
 
