Skeletal mesh build refactor, the skeletal mesh now have the capability to be rebuild without reimport like the staticmesh.
Copying //UE4/Dev-Editor-SKBuildRefactor to Dev-Editor (//UE4/Dev-Editor)

diff --git a/Engine/Plugins/Experimental/SkeletalReduction/Source/Private/SkeletalMeshReductionPlugin.cpp b/Engine/Plugins/Experimental/SkeletalReduction/Source/Private/SkeletalMeshReductionPlugin.cpp
--- a/Engine/Plugins/Experimental/SkeletalReduction/Source/Private/SkeletalMeshReductionPlugin.cpp
+++ b/Engine/Plugins/Experimental/SkeletalReduction/Source/Private/SkeletalMeshReductionPlugin.cpp
@@ -100,8 +100,7 @@ public:
 	* @returns true if reduction was successful.
 	*/
 	virtual bool ReduceSkeletalMesh( class USkeletalMesh* SkeletalMesh,
-		                             int32 LODIndex,
-		                             bool bReregisterComponent = true
+		                             int32 LODIndex		                             
 	                                ) override ;
 
 
@@ -170,7 +169,8 @@ private:
 		                         FSkeletalMeshOptimizationSettings Settings,
 								 const FImportantBones& ImportantBones,
 		                         const TArray<FMatrix>& BoneMatrices,
-		                         const int32 LODIndex ) const;
+		                         const int32 LODIndex,
+								 const bool bReducingSourceModel) const;
 
 	/**
 	* Remove the specified section from the mesh.
@@ -204,7 +204,8 @@ private:
 		                                 const FSkeletalMeshLODModel& SrcLODModel,
 		                                 const SkeletalSimplifier::FSkinnedSkeletalMesh& SkinnedMesh,
 		                                 const FReferenceSkeleton& RefSkeleton,
-		                                 FSkeletalMeshLODModel& NewModel ) const;
+		                                 FSkeletalMeshLODModel& NewModel,
+										 const bool bReducingSourceModel) const;
 
 	/**
 	* Add the SourceModelInflunces to the LODModel in the case that alternate weights exists.
@@ -359,10 +360,7 @@ void FSkeletalMeshReduction::ShutdownModule()
 
 }
 
-
-
-
-bool FQuadricSkeletalMeshReduction::ReduceSkeletalMesh( USkeletalMesh* SkeletalMesh, int32 LODIndex, bool bReregisterComponent )
+bool FQuadricSkeletalMeshReduction::ReduceSkeletalMesh( USkeletalMesh* SkeletalMesh, int32 LODIndex )
 {
 	check(SkeletalMesh);
 	check(LODIndex >= 0);
@@ -372,21 +370,8 @@ bool FQuadricSkeletalMeshReduction::ReduceSkeletalMesh( USkeletalMesh* SkeletalM
 	check(SkeletalMeshResource);
 	check(LODIndex <= SkeletalMeshResource->LODModels.Num());
 
-	if (bReregisterComponent)
-	{
-		TComponentReregisterContext<USkinnedMeshComponent> ReregisterContext;
-		SkeletalMesh->ReleaseResources();
-		SkeletalMesh->ReleaseResourcesFence.Wait();
-
-		ReduceSkeletalMesh(*SkeletalMesh, *SkeletalMeshResource, LODIndex);
-
-		SkeletalMesh->PostEditChange();
-		SkeletalMesh->InitResources();
-	}
-	else
-	{
-		ReduceSkeletalMesh(*SkeletalMesh, *SkeletalMeshResource, LODIndex);
-	}
+	FScopedSkeletalMeshPostEditChange ScopedPostEditChange(SkeletalMesh);
+	ReduceSkeletalMesh(*SkeletalMesh, *SkeletalMeshResource, LODIndex);
 
 	return true;
 }
@@ -1541,8 +1526,8 @@ void FQuadricSkeletalMeshReduction::ConvertToFSkeletalMeshLODModel( const int32
 	                                                                const FSkeletalMeshLODModel& SrcLODModel,
 	                                                                const SkeletalSimplifier::FSkinnedSkeletalMesh& SkinnedMesh,
 	                                                                const FReferenceSkeleton& RefSkeleton,
-	                                                                FSkeletalMeshLODModel& NewModel
-                                                                   ) const
+	                                                                FSkeletalMeshLODModel& NewModel,	
+																	const bool bReducingSourceModel) const
 {
 	// We might be re-using this model - so clear it.
 
@@ -1553,11 +1538,14 @@ void FQuadricSkeletalMeshReduction::ConvertToFSkeletalMeshLODModel( const int32
 	FSkeletalMeshData SkeletalMeshData;
 	ExtractFSkeletalData(SkinnedMesh, SkeletalMeshData);
 
-	// Add alternate weight data to the NewModel.  Note, this has to be done before we "BuildSkeletalMesh" to insure
-	// that the bone-based vertex chunking respects the alternate weights.
-	// NB: this only prepares the NewModel, but BuildSkeletalMesh is only 1/2-aware of this, so we will have to do some additional work after.
-	
-	AddSourceModelInfluences(SrcLODModel, SkinnedMesh, NewModel);
+	//if (!bReducingSourceModel)
+	{
+		// Add alternate weight data to the NewModel.  Note, this has to be done before we "BuildSkeletalMesh" to insure
+		// that the bone-based vertex chunking respects the alternate weights.
+		// NB: this only prepares the NewModel, but BuildSkeletalMesh is only 1/2-aware of this, so we will have to do some additional work after.
+
+		AddSourceModelInfluences(SrcLODModel, SkinnedMesh, NewModel);
+	}
 	
 
 
@@ -1576,6 +1564,7 @@ void FQuadricSkeletalMeshReduction::ConvertToFSkeletalMeshLODModel( const int32
 	Options.bComputeNormals = false;
 	Options.bComputeTangents = false;
 	Options.bUseMikkTSpace = true; //Avoid builtin build by specifying true for mikkt space
+	Options.bComputeWeightedNormals = false;
 	Options.OverlappingThresholds.ThresholdPosition = 0.0f;
 	Options.OverlappingThresholds.ThresholdTangentNormal = 0.0f;
 	Options.OverlappingThresholds.ThresholdUV = 0.0f;
@@ -1595,15 +1584,19 @@ void FQuadricSkeletalMeshReduction::ConvertToFSkeletalMeshLODModel( const int32
 		Options
 	);
 
+	//Re-Apply the user section changes, the UserSectionsData is map to original section and should match the builded LODModel
+	NewModel.SyncronizeUserSectionsDataArray();
+
 	// Set texture coordinate count on the new model.
 	NewModel.NumTexCoords = SkeletalMeshData.TexCoordCount;
 
 
+	//if (!bReducingSourceModel)
+	{
+		// Update the alternate weights
 
-	// Update the alternate weights
-	
-	UpdateAlternateWeights(MaxBonesPerVertex, NewModel);
-
+		UpdateAlternateWeights(MaxBonesPerVertex, NewModel);
+	}
 }
 
 bool FQuadricSkeletalMeshReduction::ReduceSkeletalLODModel( const FSkeletalMeshLODModel& SrcModel,
@@ -1613,7 +1606,8 @@ bool FQuadricSkeletalMeshReduction::ReduceSkeletalLODModel( const FSkeletalMeshL
 	                                                        FSkeletalMeshOptimizationSettings Settings,
 															const FImportantBones& ImportantBones,
 	                                                        const TArray<FMatrix>& BoneMatrices,
-	                                                        const int32 LODIndex
+	                                                        const int32 LODIndex,
+															const bool bReducingSourceModel
                                                            ) const
 {
 
@@ -1682,7 +1676,7 @@ bool FQuadricSkeletalMeshReduction::ReduceSkeletalLODModel( const FSkeletalMeshL
 
 		// Convert to SkeletalMeshLODModel. 
 
-		ConvertToFSkeletalMeshLODModel(Settings.MaxBonesPerVertex, SrcModel, SkinnedSkeletalMesh, RefSkeleton, OutSkeletalMeshLODModel);
+		ConvertToFSkeletalMeshLODModel(Settings.MaxBonesPerVertex, SrcModel, SkinnedSkeletalMesh, RefSkeleton, OutSkeletalMeshLODModel, bReducingSourceModel);
 
 		// We may need to do additional simplification if the user specified a hard number limit for verts and
 		// the internal chunking during conversion split some verts.
@@ -1806,11 +1800,23 @@ void FQuadricSkeletalMeshReduction::ReduceSkeletalMesh(USkeletalMesh& SkeletalMe
 		uint16 MaterialIndex;
 		bool bCastShadow;
 		bool bRecomputeTangent;
+		bool bDisabled;
+		int32 GenerateUpToLodIndex;
+		int32 ChunkedParentSectionIndex;
+		int32 OriginalDataSectionIndex;
 	};
 
 	TMap<int32, FSectionData> BackupSectionIndexToSectionData;
+	TMap<int32, FSkelMeshSourceSectionUserData> BackupUserSectionsData;
+	FString BackupLodModelBuildStringID = TEXT("");
 	//Store the section data. Store the source LOD model in case we add a LOD, or the target LOD model in case the LOD already exist
 	FSkeletalMeshLODModel& BackupSectionLODModel = bLODModelAdded ? SkeletalMeshResource.LODModels[Settings.BaseLOD] : SkeletalMeshResource.LODModels[LODIndex];
+	if (!bLODModelAdded)
+	{
+		BackupLodModelBuildStringID = BackupSectionLODModel.BuildStringID;
+		BackupUserSectionsData = BackupSectionLODModel.UserSectionsData;
+	}
+
 	BackupSectionIndexToSectionData.Reserve(BackupSectionLODModel.Sections.Num());
 
 	for (int32 SectionIndex = 0; SectionIndex < BackupSectionLODModel.Sections.Num(); ++SectionIndex)
@@ -1825,23 +1831,21 @@ void FQuadricSkeletalMeshReduction::ReduceSkeletalMesh(USkeletalMesh& SkeletalMe
 		SectionData.MaterialIndex = BackupSectionLODModel.Sections[SectionIndex].MaterialIndex;
 		SectionData.bCastShadow = BackupSectionLODModel.Sections[SectionIndex].bCastShadow;
 		SectionData.bRecomputeTangent = BackupSectionLODModel.Sections[SectionIndex].bRecomputeTangent;
+		SectionData.bDisabled = BackupSectionLODModel.Sections[SectionIndex].bDisabled;
+		SectionData.GenerateUpToLodIndex = BackupSectionLODModel.Sections[SectionIndex].GenerateUpToLodIndex;
+		SectionData.ChunkedParentSectionIndex = BackupSectionLODModel.Sections[SectionIndex].ChunkedParentSectionIndex;
+		SectionData.OriginalDataSectionIndex = BackupSectionLODModel.Sections[SectionIndex].OriginalDataSectionIndex;
 	}
 
-	auto FillClothingData = [&SkeletalMeshResource, &LODIndex, bLODModelAdded](int32 &EnableSectionNumber, TArray<bool> &SectionStatus)
+	auto FillSectionMaterialSlot = [&SkeletalMeshResource, &LODIndex, bLODModelAdded](TArray<int32>& SectionMaterialSlot)
 	{
-		EnableSectionNumber = 0;
-		SectionStatus.Empty();
+		SectionMaterialSlot.Empty();
 		if (!bLODModelAdded && SkeletalMeshResource.LODModels.IsValidIndex(LODIndex))
 		{
 			int32 SectionNumber = SkeletalMeshResource.LODModels[LODIndex].Sections.Num();
-			SectionStatus.Reserve(SectionNumber);
 			for (int32 SectionIndex = 0; SectionIndex < SectionNumber; ++SectionIndex)
 			{
-				SectionStatus.Add(!SkeletalMeshResource.LODModels[LODIndex].Sections[SectionIndex].bDisabled);
-				if (SectionStatus[SectionIndex])
-				{
-					EnableSectionNumber++;
-				}
+				SectionMaterialSlot.Add(SkeletalMeshResource.LODModels[LODIndex].Sections[SectionIndex].OriginalDataSectionIndex);
 			}
 		}
 	};
@@ -1849,31 +1853,20 @@ void FQuadricSkeletalMeshReduction::ReduceSkeletalMesh(USkeletalMesh& SkeletalMe
 	// Unbind any existing clothing assets before we reimport the geometry
 	TArray<ClothingAssetUtils::FClothingAssetMeshBinding> ClothingBindings;
 	//Get a map of enable/disable sections
-	int32 OriginalSectionNumberBeforeReduction = 0;
-	TArray<bool> OriginalSectionEnableBeforeReduction;
+	TArray<int32> OriginalSectionMaterialSlot;
 
 	//Do not play with cloth if the LOD is added
 	if (!bLODModelAdded)
 	{
-		//Store the clothBinding
-		ClothingAssetUtils::GetMeshClothingAssetBindings(&SkeletalMesh, ClothingBindings, LODIndex);
-		FillClothingData(OriginalSectionNumberBeforeReduction, OriginalSectionEnableBeforeReduction);
-		//Unbind the Cloth for this LOD before we reduce it, we will put back the cloth after the reduction, if it still match the sections
-		for (ClothingAssetUtils::FClothingAssetMeshBinding& Binding : ClothingBindings)
-		{
-			if (Binding.LODIndex == LODIndex)
-			{
-				Binding.Asset->UnbindFromSkeletalMesh(&SkeletalMesh, Binding.LODIndex);
-			}
-		}
+		FLODUtilities::UnbindClothingAndBackup(&SkeletalMesh, ClothingBindings, LODIndex);
 	}
 
 	bool bReducingSourceModel = false;
-	//Reducing Base LOD, we need to use the temporary data so it can be iterative
-	if (BaseLOD == LODIndex && SkelResource->OriginalReductionSourceMeshData.IsValidIndex(BaseLOD) && !SkelResource->OriginalReductionSourceMeshData[BaseLOD]->IsEmpty())
+	//Reducing source LOD, we need to use the temporary data so it can be iterative
+	if (BaseLOD == LODIndex && (SkelResource->OriginalReductionSourceMeshData.IsValidIndex(BaseLOD) && !SkelResource->OriginalReductionSourceMeshData[BaseLOD]->IsEmpty()))
 	{
 		TMap<FString, TArray<FMorphTargetDelta>> TempLODMorphTargetData;
-		SkelResource->OriginalReductionSourceMeshData[BaseLOD]->LoadReductionData(*SrcModel, TempLODMorphTargetData);
+		SkelResource->OriginalReductionSourceMeshData[BaseLOD]->LoadReductionData(*SrcModel, TempLODMorphTargetData, &SkeletalMesh);
 		bReducingSourceModel = true;
 	}
 	else
@@ -1955,6 +1948,8 @@ void FQuadricSkeletalMeshReduction::ReduceSkeletalMesh(USkeletalMesh& SkeletalMe
 
 	// Swap out the old model.  
 	FSkeletalMeshImportData RawMesh;
+	ESkeletalMeshGeoImportVersions GeoImportVersion = ESkeletalMeshGeoImportVersions::Before_Versionning;
+	ESkeletalMeshSkinningImportVersions SkinningImportVersion = ESkeletalMeshSkinningImportVersions::Before_Versionning;
 	{
 		FSkeletalMeshLODModel* Old = LODModels[LODIndex];
 
@@ -1962,12 +1957,21 @@ void FQuadricSkeletalMeshReduction::ReduceSkeletalMesh(USkeletalMesh& SkeletalMe
 
 		if (!bReducingSourceModel && Old)
 		{
+			//We need to backup the original RawSkeletalMeshBulkData in case it was an imported LOD
+			if (!bLODModelAdded && !Old->RawSkeletalMeshBulkData.IsEmpty())
+			{
+				Old->RawSkeletalMeshBulkData.LoadRawMesh(RawMesh);
+				GeoImportVersion = Old->RawSkeletalMeshBulkData.GeoImportVersion;
+				SkinningImportVersion = Old->RawSkeletalMeshBulkData.SkinningImportVersion;
+			}
 			delete Old;
 		}
 		else if(bReducingSourceModel)
 		{
 			//In case we reduce the source model we want to keep the original import data
 			SrcModel->RawSkeletalMeshBulkData.LoadRawMesh(RawMesh);
+			GeoImportVersion = SrcModel->RawSkeletalMeshBulkData.GeoImportVersion;
+			SkinningImportVersion = SrcModel->RawSkeletalMeshBulkData.SkinningImportVersion;
 		}
 	}
 
@@ -1975,7 +1979,7 @@ void FQuadricSkeletalMeshReduction::ReduceSkeletalMesh(USkeletalMesh& SkeletalMe
 
 	// Reduce LOD model with SrcMesh
 
-	if (ReduceSkeletalLODModel(*SrcModel, *NewModel, SkeletalMesh.GetImportedBounds(), SkeletalMesh.RefSkeleton, Settings, ImportantBones, RelativeToRefPoseMatrices, LODIndex))
+	if (ReduceSkeletalLODModel(*SrcModel, *NewModel, SkeletalMesh.GetImportedBounds(), SkeletalMesh.RefSkeleton, Settings, ImportantBones, RelativeToRefPoseMatrices, LODIndex, bReducingSourceModel))
 	{
 		// Do any joint-welding / bone removal.
 
@@ -2019,6 +2023,13 @@ void FQuadricSkeletalMeshReduction::ReduceSkeletalMesh(USkeletalMesh& SkeletalMe
 		FSkeletalMeshLODModel& ImportedModelLOD = SkeletalMesh.GetImportedModel()->LODModels[LODIndex];
 		TArray<bool> SectionMatched;
 		SectionMatched.AddZeroed(ImportedModelLOD.Sections.Num());
+		
+		if (!bLODModelAdded)
+		{
+			ImportedModelLOD.UserSectionsData = BackupUserSectionsData;
+			ImportedModelLOD.BuildStringID = BackupLodModelBuildStringID;
+		}
+
 		for (auto Kvp : BackupSectionIndexToSectionData)
 		{
 			const int32 SourceSectionIndex = Kvp.Key;
@@ -2035,6 +2046,10 @@ void FQuadricSkeletalMeshReduction::ReduceSkeletalMesh(USkeletalMesh& SkeletalMe
 				{
 					ImportedModelLOD.Sections[SectionIndex].bCastShadow = SectionData.bCastShadow;
 					ImportedModelLOD.Sections[SectionIndex].bRecomputeTangent = SectionData.bRecomputeTangent;
+					ImportedModelLOD.Sections[SectionIndex].bDisabled = SectionData.bDisabled;
+					ImportedModelLOD.Sections[SectionIndex].GenerateUpToLodIndex = SectionData.GenerateUpToLodIndex;
+					ImportedModelLOD.Sections[SectionIndex].ChunkedParentSectionIndex = SectionData.ChunkedParentSectionIndex;
+					ImportedModelLOD.Sections[SectionIndex].OriginalDataSectionIndex = SectionData.OriginalDataSectionIndex;
 					SectionMatched[SectionIndex] = true;
 					break;
 				}
@@ -2087,47 +2102,22 @@ void FQuadricSkeletalMeshReduction::ReduceSkeletalMesh(USkeletalMesh& SkeletalMe
 	if (!bLODModelAdded)
 	{
 		//Get the number of enabled section
-		int32 SectionNumberAfterReduction = 0;
-		TArray<bool> SectionEnableAfterReduction;
-		FillClothingData(SectionNumberAfterReduction, SectionEnableAfterReduction);
+		TArray<int32> SectionMaterialSlotAfterReduction;
+		FillSectionMaterialSlot(SectionMaterialSlotAfterReduction);
 
-		//Put back the clothing for this newly reduce LOD only if the section count match.
-		if (ClothingBindings.Num() > 0 && OriginalSectionNumberBeforeReduction == SectionNumberAfterReduction)
+		//Put back the clothing for this newly reduce LOD
+		if (ClothingBindings.Num() > 0)
 		{
-			TArray<int32> RemapSectionIndex;
-			int32 SectionIndexTest = 0;
-			for (int32 SectionIndexRef = 0; SectionIndexRef < OriginalSectionEnableBeforeReduction.Num(); SectionIndexRef++)
-			{
-				int32& RemapValue = RemapSectionIndex.Add_GetRef(INDEX_NONE);
-				if (!OriginalSectionEnableBeforeReduction[SectionIndexRef])
-				{
-					continue;
-				}
-				for (; SectionIndexTest <= SectionIndexRef; SectionIndexTest++)
-				{
-					if (SectionEnableAfterReduction.IsValidIndex(SectionIndexTest) && SectionEnableAfterReduction[SectionIndexTest])
-					{
-						RemapValue = SectionIndexTest++;
-						break;
-					}
-				}
-			}
-
-			for (ClothingAssetUtils::FClothingAssetMeshBinding& Binding : ClothingBindings)
-			{
-				int32 RemapBindingSectionIndex = RemapSectionIndex[Binding.SectionIndex];
-				if (RemapBindingSectionIndex != INDEX_NONE && Binding.LODIndex == LODIndex && NewModel->Sections.IsValidIndex(RemapBindingSectionIndex))
-				{
-					Binding.Asset->BindToSkeletalMesh(&SkeletalMesh, Binding.LODIndex, RemapBindingSectionIndex, Binding.AssetInternalLodIndex, false);
-				}
-			}
+			FLODUtilities::RestoreClothingFromBackup(&SkeletalMesh, ClothingBindings, LODIndex);
 		}
 	}
 
-	if (bReducingSourceModel && RawMesh.Points.Num() > 0)
+	if ((bReducingSourceModel || !bLODModelAdded ) && RawMesh.Points.Num() > 0)
 	{
 		//Put back the original import data, we need it to allow inline reduction and skeletal mesh split workflow
 		SkeletalMeshResource.LODModels[LODIndex].RawSkeletalMeshBulkData.SaveRawMesh(RawMesh);
+		SkeletalMeshResource.LODModels[LODIndex].RawSkeletalMeshBulkData.GeoImportVersion = GeoImportVersion;
+		SkeletalMeshResource.LODModels[LODIndex].RawSkeletalMeshBulkData.SkinningImportVersion = SkinningImportVersion;
 	}
 
 	SkeletalMesh.CalculateRequiredBones(SkeletalMeshResource.LODModels[LODIndex], SkeletalMesh.RefSkeleton, &BonesToRemove);
diff --git a/Engine/Source/Developer/MeshBuilder/MeshBuilder.Build.cs b/Engine/Source/Developer/MeshBuilder/MeshBuilder.Build.cs
--- a/Engine/Source/Developer/MeshBuilder/MeshBuilder.Build.cs
+++ b/Engine/Source/Developer/MeshBuilder/MeshBuilder.Build.cs
@@ -20,7 +20,11 @@ namespace UnrealBuildTool.Rules
                     "MeshDescriptionOperations",
                     "MeshReductionInterface",
                     "RawMesh",
+					"MeshUtilities",
                     "MeshUtilitiesCommon",
+                    "ClothingSystemRuntime",
+                    "MeshBoneReduction",
+                    "SkeletalMeshUtilitiesCommon",
                 }
 			);
 
diff --git a/Engine/Source/Developer/MeshUtilities/Private/MeshUtilities.cpp b/Engine/Source/Developer/MeshUtilities/Private/MeshUtilities.cpp
--- a/Engine/Source/Developer/MeshUtilities/Private/MeshUtilities.cpp
+++ b/Engine/Source/Developer/MeshUtilities/Private/MeshUtilities.cpp
@@ -645,6 +645,9 @@ void FMeshUtilities::BuildSkeletalModelFromChunks(FSkeletalMeshLODModel& LODMode
 		Section.MaterialIndex = SrcChunk->MaterialIndex;
 		Exchange(Section.BoneMap, SrcChunk->BoneMap);
 
+		Section.OriginalDataSectionIndex = SrcChunk->OriginalSectionIndex;
+		Section.ChunkedParentSectionIndex = SrcChunk->ParentChunkSectionIndex;
+
 		// Update the active bone indices on the LOD model.
 		for (int32 BoneIndex = 0; BoneIndex < Section.BoneMap.Num(); ++BoneIndex)
 		{
@@ -3450,342 +3453,6 @@ public:
 		}
 	}
 
-	void Skeletal_ComputeTangents(
-		IMeshBuildData* BuildData,
-		const FOverlappingCorners& OverlappingCorners
-		)
-	{
-		bool bBlendOverlappingNormals = true;
-		bool bIgnoreDegenerateTriangles = BuildData->BuildOptions.bRemoveDegenerateTriangles;
-		bool bComputeWeightedNormals = BuildData->BuildOptions.bComputeWeightedNormals;
-		
-		// Compute per-triangle tangents.
-		TArray<FVector> TriangleTangentX;
-		TArray<FVector> TriangleTangentY;
-		TArray<FVector> TriangleTangentZ;
-
-		Skeletal_ComputeTriangleTangents(
-			TriangleTangentX,
-			TriangleTangentY,
-			TriangleTangentZ,
-			BuildData,
-			bIgnoreDegenerateTriangles ? SMALL_NUMBER : FLT_MIN
-			);
-
-		TArray<FVector>& WedgeTangentX = BuildData->GetTangentArray(0);
-		TArray<FVector>& WedgeTangentY = BuildData->GetTangentArray(1);
-		TArray<FVector>& WedgeTangentZ = BuildData->GetTangentArray(2);
-
-		// Declare these out here to avoid reallocations.
-		TArray<FFanFace> RelevantFacesForCorner[3];
-		TArray<int32> AdjacentFaces;
-
-		int32 NumFaces = BuildData->GetNumFaces();
-		int32 NumWedges = BuildData->GetNumWedges();
-		check(NumFaces * 3 <= NumWedges);
-
-		// Allocate storage for tangents if none were provided.
-		if (WedgeTangentX.Num() != NumWedges)
-		{
-			WedgeTangentX.Empty(NumWedges);
-			WedgeTangentX.AddZeroed(NumWedges);
-		}
-		if (WedgeTangentY.Num() != NumWedges)
-		{
-			WedgeTangentY.Empty(NumWedges);
-			WedgeTangentY.AddZeroed(NumWedges);
-		}
-		if (WedgeTangentZ.Num() != NumWedges)
-		{
-			WedgeTangentZ.Empty(NumWedges);
-			WedgeTangentZ.AddZeroed(NumWedges);
-		}
-
-		for (int32 FaceIndex = 0; FaceIndex < NumFaces; FaceIndex++)
-		{
-			int32 WedgeOffset = FaceIndex * 3;
-			FVector CornerPositions[3];
-			FVector CornerTangentX[3];
-			FVector CornerTangentY[3];
-			FVector CornerTangentZ[3];
-
-			for (int32 CornerIndex = 0; CornerIndex < 3; CornerIndex++)
-			{
-				CornerTangentX[CornerIndex] = FVector::ZeroVector;
-				CornerTangentY[CornerIndex] = FVector::ZeroVector;
-				CornerTangentZ[CornerIndex] = FVector::ZeroVector;
-				CornerPositions[CornerIndex] = BuildData->GetVertexPosition(FaceIndex, CornerIndex);
-				RelevantFacesForCorner[CornerIndex].Reset();
-			}
-
-			// Don't process degenerate triangles.
-			if (PointsEqual(CornerPositions[0], CornerPositions[1], BuildData->BuildOptions.OverlappingThresholds)
-				|| PointsEqual(CornerPositions[0], CornerPositions[2], BuildData->BuildOptions.OverlappingThresholds)
-				|| PointsEqual(CornerPositions[1], CornerPositions[2], BuildData->BuildOptions.OverlappingThresholds))
-			{
-				continue;
-			}
-
-			// No need to process triangles if tangents already exist.
-			bool bCornerHasTangents[3] = { 0 };
-			for (int32 CornerIndex = 0; CornerIndex < 3; CornerIndex++)
-			{
-				bCornerHasTangents[CornerIndex] = !WedgeTangentX[WedgeOffset + CornerIndex].IsZero()
-					&& !WedgeTangentY[WedgeOffset + CornerIndex].IsZero()
-					&& !WedgeTangentZ[WedgeOffset + CornerIndex].IsZero();
-			}
-			if (bCornerHasTangents[0] && bCornerHasTangents[1] && bCornerHasTangents[2])
-			{
-				continue;
-			}
-
-			// Calculate smooth vertex normals.
-			float Determinant = FVector::Triple(
-				TriangleTangentX[FaceIndex],
-				TriangleTangentY[FaceIndex],
-				TriangleTangentZ[FaceIndex]
-				);
-
-			// Start building a list of faces adjacent to this face.
-			AdjacentFaces.Reset();
-			for (int32 CornerIndex = 0; CornerIndex < 3; CornerIndex++)
-			{
-				int32 ThisCornerIndex = WedgeOffset + CornerIndex;
-				const TArray<int32>& DupVerts = OverlappingCorners.FindIfOverlapping(ThisCornerIndex);
-				if (DupVerts.Num() == 0)
-				{
-					AdjacentFaces.AddUnique(ThisCornerIndex / 3); // I am a "dup" of myself
-				}
-				for (int32 k = 0; k < DupVerts.Num(); k++)
-				{
-					AdjacentFaces.AddUnique(DupVerts[k] / 3);
-				}
-			}
-
-			// We need to sort these here because the criteria for point equality is
-			// exact, so we must ensure the exact same order for all dups.
-			AdjacentFaces.Sort();
-
-			// Process adjacent faces
-			for (int32 AdjacentFaceIndex = 0; AdjacentFaceIndex < AdjacentFaces.Num(); AdjacentFaceIndex++)
-			{
-				int32 OtherFaceIndex = AdjacentFaces[AdjacentFaceIndex];
-				for (int32 OurCornerIndex = 0; OurCornerIndex < 3; OurCornerIndex++)
-				{
-					if (bCornerHasTangents[OurCornerIndex])
-						continue;
-
-					FFanFace NewFanFace;
-					int32 CommonIndexCount = 0;
-
-					// Check for vertices in common.
-					if (FaceIndex == OtherFaceIndex)
-					{
-						CommonIndexCount = 3;
-						NewFanFace.LinkedVertexIndex = OurCornerIndex;
-					}
-					else
-					{
-						// Check matching vertices against main vertex .
-						for (int32 OtherCornerIndex = 0; OtherCornerIndex < 3; OtherCornerIndex++)
-						{
-							if (PointsEqual(
-								CornerPositions[OurCornerIndex],
-								BuildData->GetVertexPosition(OtherFaceIndex, OtherCornerIndex),
-								BuildData->BuildOptions.OverlappingThresholds
-								))
-							{
-								CommonIndexCount++;
-								NewFanFace.LinkedVertexIndex = OtherCornerIndex;
-							}
-						}
-					}
-
-					// Add if connected by at least one point. Smoothing matches are considered later.
-					if (CommonIndexCount > 0)
-					{
-						NewFanFace.FaceIndex = OtherFaceIndex;
-						NewFanFace.bFilled = (OtherFaceIndex == FaceIndex); // Starter face for smoothing floodfill.
-						NewFanFace.bBlendTangents = NewFanFace.bFilled;
-						NewFanFace.bBlendNormals = NewFanFace.bFilled;
-						RelevantFacesForCorner[OurCornerIndex].Add(NewFanFace);
-					}
-				}
-			}
-
-			// Find true relevance of faces for a vertex normal by traversing
-			// smoothing-group-compatible connected triangle fans around common vertices.
-			for (int32 CornerIndex = 0; CornerIndex < 3; CornerIndex++)
-			{
-				if (bCornerHasTangents[CornerIndex])
-					continue;
-
-				int32 NewConnections;
-				do
-				{
-					NewConnections = 0;
-					for (int32 OtherFaceIdx = 0; OtherFaceIdx < RelevantFacesForCorner[CornerIndex].Num(); OtherFaceIdx++)
-					{
-						FFanFace& OtherFace = RelevantFacesForCorner[CornerIndex][OtherFaceIdx];
-						// The vertex' own face is initially the only face with bFilled == true.
-						if (OtherFace.bFilled)
-						{
-							for (int32 NextFaceIndex = 0; NextFaceIndex < RelevantFacesForCorner[CornerIndex].Num(); NextFaceIndex++)
-							{
-								FFanFace& NextFace = RelevantFacesForCorner[CornerIndex][NextFaceIndex];
-								if (!NextFace.bFilled) // && !NextFace.bBlendTangents)
-								{
-									if ((NextFaceIndex != OtherFaceIdx)
-										&& (BuildData->GetFaceSmoothingGroups(NextFace.FaceIndex) & BuildData->GetFaceSmoothingGroups(OtherFace.FaceIndex)))
-									{
-										int32 CommonVertices = 0;
-										int32 CommonTangentVertices = 0;
-										int32 CommonNormalVertices = 0;
-										for (int32 OtherCornerIndex = 0; OtherCornerIndex < 3; OtherCornerIndex++)
-										{
-											for (int32 NextCornerIndex = 0; NextCornerIndex < 3; NextCornerIndex++)
-											{
-												int32 NextVertexIndex = BuildData->GetVertexIndex(NextFace.FaceIndex, NextCornerIndex);
-												int32 OtherVertexIndex = BuildData->GetVertexIndex(OtherFace.FaceIndex, OtherCornerIndex);
-												if (PointsEqual(
-													BuildData->GetVertexPosition(NextFace.FaceIndex, NextCornerIndex),
-													BuildData->GetVertexPosition(OtherFace.FaceIndex, OtherCornerIndex),
-													BuildData->BuildOptions.OverlappingThresholds))
-												{
-													CommonVertices++;
-
-
-													if (UVsEqual(
-														BuildData->GetVertexUV(NextFace.FaceIndex, NextCornerIndex, 0),
-														BuildData->GetVertexUV(OtherFace.FaceIndex, OtherCornerIndex, 0),
-														BuildData->BuildOptions.OverlappingThresholds))
-													{
-														CommonTangentVertices++;
-													}
-													if (bBlendOverlappingNormals
-														|| NextVertexIndex == OtherVertexIndex)
-													{
-														CommonNormalVertices++;
-													}
-												}
-											}
-										}
-										// Flood fill faces with more than one common vertices which must be touching edges.
-										if (CommonVertices > 1)
-										{
-											NextFace.bFilled = true;
-											NextFace.bBlendNormals = (CommonNormalVertices > 1);
-											NewConnections++;
-
-											// Only blend tangents if there is no UV seam along the edge with this face.
-											if (OtherFace.bBlendTangents && CommonTangentVertices > 1)
-											{
-												float OtherDeterminant = FVector::Triple(
-													TriangleTangentX[NextFace.FaceIndex],
-													TriangleTangentY[NextFace.FaceIndex],
-													TriangleTangentZ[NextFace.FaceIndex]
-													);
-												if ((Determinant * OtherDeterminant) > 0.0f)
-												{
-													NextFace.bBlendTangents = true;
-												}
-											}
-										}
-									}
-								}
-							}
-						}
-					}
-				}
-				while (NewConnections > 0);
-			}
-
-			// Vertex normal construction.
-			for (int32 CornerIndex = 0; CornerIndex < 3; CornerIndex++)
-			{
-				if (bCornerHasTangents[CornerIndex])
-				{
-					CornerTangentX[CornerIndex] = WedgeTangentX[WedgeOffset + CornerIndex];
-					CornerTangentY[CornerIndex] = WedgeTangentY[WedgeOffset + CornerIndex];
-					CornerTangentZ[CornerIndex] = WedgeTangentZ[WedgeOffset + CornerIndex];
-				}
-				else
-				{
-					for (int32 RelevantFaceIdx = 0; RelevantFaceIdx < RelevantFacesForCorner[CornerIndex].Num(); RelevantFaceIdx++)
-					{
-						FFanFace const& RelevantFace = RelevantFacesForCorner[CornerIndex][RelevantFaceIdx];
-						if (RelevantFace.bFilled)
-						{
-							int32 OtherFaceIndex = RelevantFace.FaceIndex;
-							float CornerWeight = 1.0f;
-							if (bComputeWeightedNormals)
-							{
-
-								FVector OtherFacePoint[3] = { BuildData->GetVertexPosition(OtherFaceIndex, 0), BuildData->GetVertexPosition(OtherFaceIndex, 1), BuildData->GetVertexPosition(OtherFaceIndex, 2) };
-								float OtherFaceArea = TriangleUtilities::ComputeTriangleArea(OtherFacePoint[0], OtherFacePoint[1], OtherFacePoint[2]);
-								int32 OtherFaceCornerIndex = RelevantFace.LinkedVertexIndex;
-								float OtherFaceAngle = TriangleUtilities::ComputeTriangleCornerAngle(OtherFacePoint[OtherFaceCornerIndex], OtherFacePoint[(OtherFaceCornerIndex + 1) % 3], OtherFacePoint[(OtherFaceCornerIndex + 2) % 3]);
-								//Get the CornerWeight
-								CornerWeight = OtherFaceArea * OtherFaceAngle;
-							}
-
-							if (RelevantFace.bBlendTangents)
-							{
-								CornerTangentX[CornerIndex] += CornerWeight * TriangleTangentX[OtherFaceIndex];
-								CornerTangentY[CornerIndex] += CornerWeight * TriangleTangentY[OtherFaceIndex];
-							}
-							if (RelevantFace.bBlendNormals)
-							{
-								CornerTangentZ[CornerIndex] += CornerWeight * TriangleTangentZ[OtherFaceIndex];
-							}
-						}
-					}
-					if (!WedgeTangentX[WedgeOffset + CornerIndex].IsZero())
-					{
-						CornerTangentX[CornerIndex] = WedgeTangentX[WedgeOffset + CornerIndex];
-					}
-					if (!WedgeTangentY[WedgeOffset + CornerIndex].IsZero())
-					{
-						CornerTangentY[CornerIndex] = WedgeTangentY[WedgeOffset + CornerIndex];
-					}
-					if (!WedgeTangentZ[WedgeOffset + CornerIndex].IsZero())
-					{
-						CornerTangentZ[CornerIndex] = WedgeTangentZ[WedgeOffset + CornerIndex];
-					}
-				}
-			}
-
-			// Normalization.
-			for (int32 CornerIndex = 0; CornerIndex < 3; CornerIndex++)
-			{
-				CornerTangentX[CornerIndex].Normalize();
-				CornerTangentY[CornerIndex].Normalize();
-				CornerTangentZ[CornerIndex].Normalize();
-
-				// Gram-Schmidt orthogonalization
-				CornerTangentY[CornerIndex] -= CornerTangentX[CornerIndex] * (CornerTangentX[CornerIndex] | CornerTangentY[CornerIndex]);
-				CornerTangentY[CornerIndex].Normalize();
-
-				CornerTangentX[CornerIndex] -= CornerTangentZ[CornerIndex] * (CornerTangentZ[CornerIndex] | CornerTangentX[CornerIndex]);
-				CornerTangentX[CornerIndex].Normalize();
-				CornerTangentY[CornerIndex] -= CornerTangentZ[CornerIndex] * (CornerTangentZ[CornerIndex] | CornerTangentY[CornerIndex]);
-				CornerTangentY[CornerIndex].Normalize();
-			}
-
-			// Copy back to the mesh.
-			for (int32 CornerIndex = 0; CornerIndex < 3; CornerIndex++)
-			{
-				WedgeTangentX[WedgeOffset + CornerIndex] = CornerTangentX[CornerIndex];
-				WedgeTangentY[WedgeOffset + CornerIndex] = CornerTangentY[CornerIndex];
-				WedgeTangentZ[WedgeOffset + CornerIndex] = CornerTangentZ[CornerIndex];
-			}
-		}
-
-		check(WedgeTangentX.Num() == NumWedges);
-		check(WedgeTangentY.Num() == NumWedges);
-		check(WedgeTangentZ.Num() == NumWedges);
-	}
-
 	bool IsTriangleMirror(IMeshBuildData* BuildData, const TArray<FVector>& TriangleTangentZ, const uint32 FaceIdxA, const uint32 FaceIdxB)
 	{
 		if (FaceIdxA == FaceIdxB)
@@ -3819,7 +3486,7 @@ public:
 		return false;
 	}
 
-	void Skeletal_ComputeTangents_MikkTSpace(
+	void Skeletal_ComputeTangents(
 		IMeshBuildData* BuildData,
 		const FOverlappingCorners& OverlappingCorners
 		)
@@ -3827,6 +3494,7 @@ public:
 		bool bBlendOverlappingNormals = true;
 		bool bIgnoreDegenerateTriangles = BuildData->BuildOptions.bRemoveDegenerateTriangles;
 		bool bComputeWeightedNormals = BuildData->BuildOptions.bComputeWeightedNormals;
+		bool bUseMikktSpace = BuildData->BuildOptions.bUseMikkTSpace && (BuildData->BuildOptions.bComputeNormals || BuildData->BuildOptions.bComputeTangents);
 		
 		int32 NumFaces = BuildData->GetNumFaces();
 		int32 NumWedges = BuildData->GetNumWedges();
@@ -3859,25 +3527,23 @@ public:
 		TArray<FFanFace> RelevantFacesForCorner[3];
 		TArray<int32> AdjacentFaces;
 
-		bool bWedgeTSpace = false;
-
-		if (WedgeTangentX.Num() > 0 && WedgeTangentY.Num() > 0)
+		// Allocate storage for tangents and normal if none were provided.
+		if (WedgeTangentX.Num() != NumWedges)
 		{
-			bWedgeTSpace = true;
-			for (int32 WedgeIdx = 0; WedgeIdx < WedgeTangentX.Num()
-				&& WedgeIdx < WedgeTangentY.Num(); ++WedgeIdx)
-			{
-				bWedgeTSpace = bWedgeTSpace && (!WedgeTangentX[WedgeIdx].IsNearlyZero()) && (!WedgeTangentY[WedgeIdx].IsNearlyZero());
-			}
+			WedgeTangentX.Empty(NumWedges);
+			WedgeTangentX.AddZeroed(NumWedges);
+		}
+		if (WedgeTangentY.Num() != NumWedges)
+		{
+			WedgeTangentY.Empty(NumWedges);
+			WedgeTangentY.AddZeroed(NumWedges);
 		}
-
-		// Allocate storage for tangents if none were provided, and calculate normals for MikkTSpace.
 		if (WedgeTangentZ.Num() != NumWedges)
 		{
-			// normals are not included, so we should calculate them
 			WedgeTangentZ.Empty(NumWedges);
 			WedgeTangentZ.AddZeroed(NumWedges);
 		}
+
 		bool bIsZeroLengthNormalErrorMessageDisplayed = false;
 		// we need to calculate normals for MikkTSpace
 		for (int32 FaceIndex = 0; FaceIndex < NumFaces; FaceIndex++)
@@ -3885,10 +3551,14 @@ public:
 			int32 PatchIndex = FaceIndexToPatchIndex[FaceIndex];
 			int32 WedgeOffset = FaceIndex * 3;
 			FVector CornerPositions[3];
+			FVector CornerTangentX[3];
+			FVector CornerTangentY[3];
 			FVector CornerNormal[3];
 
 			for (int32 CornerIndex = 0; CornerIndex < 3; CornerIndex++)
 			{
+				CornerTangentX[CornerIndex] = FVector::ZeroVector;
+				CornerTangentY[CornerIndex] = FVector::ZeroVector;
 				CornerNormal[CornerIndex] = FVector::ZeroVector;
 				CornerPositions[CornerIndex] = BuildData->GetVertexPosition(FaceIndex, CornerIndex);
 				RelevantFacesForCorner[CornerIndex].Reset();
@@ -3904,15 +3574,31 @@ public:
 
 			// No need to process triangles if tangents already exist.
 			bool bCornerHasNormal[3] = { 0 };
+			bool bCornerHasTangents[3] = { 0 };
+			bool bSkipNTB = true;
 			for (int32 CornerIndex = 0; CornerIndex < 3; CornerIndex++)
 			{
 				bCornerHasNormal[CornerIndex] = !WedgeTangentZ[WedgeOffset + CornerIndex].IsZero();
+				bCornerHasTangents[CornerIndex] = !WedgeTangentX[WedgeOffset + CornerIndex].IsZero() && !WedgeTangentY[WedgeOffset + CornerIndex].IsZero();
+
+				//If we want to compute mikkt we dont check tangents to skip this corner
+				if (!bCornerHasNormal[CornerIndex] || (!bUseMikktSpace && !bCornerHasTangents[CornerIndex]))
+				{
+					bSkipNTB = false;
+				}
 			}
-			if (bCornerHasNormal[0] && bCornerHasNormal[1] && bCornerHasNormal[2])
+			if (bSkipNTB)
 			{
 				continue;
 			}
 
+			// Calculate smooth vertex normals.
+			float Determinant = FVector::Triple(
+				TriangleTangentX[FaceIndex],
+				TriangleTangentY[FaceIndex],
+				TriangleTangentZ[FaceIndex]
+			);
+
 			// Start building a list of faces adjacent to this face.
 			AdjacentFaces.Reset();
 			for (int32 CornerIndex = 0; CornerIndex < 3; CornerIndex++)
@@ -3946,7 +3632,7 @@ public:
 				int32 OtherFaceIndex = AdjacentFaces[AdjacentFaceIndex];
 				for (int32 OurCornerIndex = 0; OurCornerIndex < 3; OurCornerIndex++)
 				{
-					if (bCornerHasNormal[OurCornerIndex])
+					if (bCornerHasNormal[OurCornerIndex] && (bUseMikktSpace || bCornerHasTangents[OurCornerIndex]))
 						continue;
 
 					FFanFace NewFanFace;
@@ -3991,7 +3677,7 @@ public:
 			// smoothing-group-compatible connected triangle fans around common vertices.
 			for (int32 CornerIndex = 0; CornerIndex < 3; CornerIndex++)
 			{
-				if (bCornerHasNormal[CornerIndex])
+				if (bCornerHasNormal[CornerIndex] && (bUseMikktSpace || bCornerHasTangents[CornerIndex]))
 					continue;
 
 				int32 NewConnections;
@@ -4013,6 +3699,7 @@ public:
 											&& (BuildData->GetFaceSmoothingGroups(NextFace.FaceIndex) & BuildData->GetFaceSmoothingGroups(OtherFace.FaceIndex)))
 									{
 										int32 CommonVertices = 0;
+										int32 CommonTangentVertices = 0;
 										int32 CommonNormalVertices = 0;
 										for (int32 OtherCornerIndex = 0; OtherCornerIndex < 3; OtherCornerIndex++)
 										{
@@ -4026,6 +3713,13 @@ public:
 													BuildData->BuildOptions.OverlappingThresholds))
 												{
 													CommonVertices++;
+													if (!bUseMikktSpace && UVsEqual(
+														BuildData->GetVertexUV(NextFace.FaceIndex, NextCornerIndex, 0),
+														BuildData->GetVertexUV(OtherFace.FaceIndex, OtherCornerIndex, 0),
+														BuildData->BuildOptions.OverlappingThresholds))
+													{
+														CommonTangentVertices++;
+													}
 													if (bBlendOverlappingNormals
 														|| NextVertexIndex == OtherVertexIndex)
 													{
@@ -4040,6 +3734,19 @@ public:
 											NextFace.bFilled = true;
 											NextFace.bBlendNormals = (CommonNormalVertices > 1);
 											NewConnections++;
+											// Only blend tangents if there is no UV seam along the edge with this face.
+											if (!bUseMikktSpace && OtherFace.bBlendTangents && CommonTangentVertices > 1)
+											{
+												float OtherDeterminant = FVector::Triple(
+													TriangleTangentX[NextFace.FaceIndex],
+													TriangleTangentY[NextFace.FaceIndex],
+													TriangleTangentZ[NextFace.FaceIndex]
+												);
+												if ((Determinant * OtherDeterminant) > 0.0f)
+												{
+													NextFace.bBlendTangents = true;
+												}
+											}
 										}
 									}
 								}
@@ -4053,11 +3760,19 @@ public:
 			// Vertex normal construction.
 			for (int32 CornerIndex = 0; CornerIndex < 3; CornerIndex++)
 			{
-				if (bCornerHasNormal[CornerIndex])
+				bool bUseProvidedNormal = bCornerHasNormal[CornerIndex];
+				bool bUseProvidedTangents = !bUseMikktSpace && bCornerHasTangents[CornerIndex];
+				if (bUseProvidedNormal)
 				{
 					CornerNormal[CornerIndex] = WedgeTangentZ[WedgeOffset + CornerIndex];
 				}
-				else
+				if (bUseProvidedTangents)
+				{
+					CornerTangentX[CornerIndex] = WedgeTangentX[WedgeOffset + CornerIndex];
+					CornerTangentY[CornerIndex] = WedgeTangentY[WedgeOffset + CornerIndex];
+				}
+
+				if (!bUseProvidedNormal || !bUseProvidedTangents)
 				{
 					for (int32 RelevantFaceIdx = 0; RelevantFaceIdx < RelevantFacesForCorner[CornerIndex].Num(); RelevantFaceIdx++)
 					{
@@ -4075,13 +3790,25 @@ public:
 								//Get the CornerWeight
 								CornerWeight = OtherFaceArea * OtherFaceAngle;
 							}
-							
-							if (RelevantFace.bBlendNormals)
+							if (!bUseProvidedTangents && RelevantFace.bBlendTangents)
+							{
+								CornerTangentX[CornerIndex] += CornerWeight * TriangleTangentX[OtherFaceIndex];
+								CornerTangentY[CornerIndex] += CornerWeight * TriangleTangentY[OtherFaceIndex];
+							}
+							if (!bUseProvidedNormal && RelevantFace.bBlendNormals)
 							{
 								CornerNormal[CornerIndex] += CornerWeight * TriangleTangentZ[OtherFaceIndex];
 							}
 						}
 					}
+					if (!WedgeTangentX[WedgeOffset + CornerIndex].IsZero())
+					{
+						CornerTangentX[CornerIndex] = WedgeTangentX[WedgeOffset + CornerIndex];
+					}
+					if (!WedgeTangentY[WedgeOffset + CornerIndex].IsZero())
+					{
+						CornerTangentY[CornerIndex] = WedgeTangentY[WedgeOffset + CornerIndex];
+					}
 					if (!WedgeTangentZ[WedgeOffset + CornerIndex].IsZero())
 					{
 						CornerNormal[CornerIndex] = WedgeTangentZ[WedgeOffset + CornerIndex];
@@ -4093,15 +3820,27 @@ public:
 			for (int32 CornerIndex = 0; CornerIndex < 3; CornerIndex++)
 			{
 				CornerNormal[CornerIndex].Normalize();
+				if (!bUseMikktSpace)
+				{
+					CornerTangentX[CornerIndex].Normalize();
+					CornerTangentY[CornerIndex].Normalize();
+
+					// Gram-Schmidt orthogonalization
+					CornerTangentY[CornerIndex] -= CornerTangentX[CornerIndex] * (CornerTangentX[CornerIndex] | CornerTangentY[CornerIndex]);
+					CornerTangentY[CornerIndex].Normalize();
+
+					CornerTangentX[CornerIndex] -= CornerNormal[CornerIndex] * (CornerNormal[CornerIndex] | CornerTangentX[CornerIndex]);
+					CornerTangentX[CornerIndex].Normalize();
+					CornerTangentY[CornerIndex] -= CornerNormal[CornerIndex] * (CornerNormal[CornerIndex] | CornerTangentY[CornerIndex]);
+					CornerTangentY[CornerIndex].Normalize();
+				}
 			}
 
-			// Copy back to the mesh.
-			for (int32 CornerIndex = 0; CornerIndex < 3; CornerIndex++)
+			auto VerifyTangentSpace = [&bIsZeroLengthNormalErrorMessageDisplayed, &BuildData](FVector& NormalizedVector)
 			{
-				//Make sure the normal do not contain nan, if it contain nan change it to a valid zero vector
-				if (CornerNormal[CornerIndex].IsNearlyZero() || CornerNormal[CornerIndex].ContainsNaN())
+				if (NormalizedVector.IsNearlyZero() || NormalizedVector.ContainsNaN())
 				{
-					CornerNormal[CornerIndex] = FVector::ZeroVector;
+					NormalizedVector = FVector::ZeroVector;
 					//We also notify the log that we compute a zero length normals, so the user is aware of it
 					if (!bIsZeroLengthNormalErrorMessageDisplayed)
 					{
@@ -4122,22 +3861,26 @@ public:
 						}
 					}
 				}
+			};
+
+			// Copy back to the mesh.
+			for (int32 CornerIndex = 0; CornerIndex < 3; CornerIndex++)
+			{
+				VerifyTangentSpace(CornerNormal[CornerIndex]);
 				WedgeTangentZ[WedgeOffset + CornerIndex] = CornerNormal[CornerIndex];
-			}
-		}
 
-		if (WedgeTangentX.Num() != NumWedges)
-		{
-			WedgeTangentX.Empty(NumWedges);
-			WedgeTangentX.AddZeroed(NumWedges);
-		}
-		if (WedgeTangentY.Num() != NumWedges)
-		{
-			WedgeTangentY.Empty(NumWedges);
-			WedgeTangentY.AddZeroed(NumWedges);
+				if (!bUseMikktSpace)
+				{
+					VerifyTangentSpace(CornerTangentX[CornerIndex]);
+					WedgeTangentX[WedgeOffset + CornerIndex] = CornerTangentX[CornerIndex];
+
+					VerifyTangentSpace(CornerTangentY[CornerIndex]);
+					WedgeTangentY[WedgeOffset + CornerIndex] = CornerTangentY[CornerIndex];
+				}
+			}
 		}
 
-		//if (!bWedgeTSpace)
+		if (bUseMikktSpace)
 		{
 			// we can use mikktspace to calculate the tangents		
 			SMikkTSpaceContext MikkTContext;
@@ -4190,15 +3933,8 @@ public:
 			TangentZ.AddZeroed(NumWedges);
 		}
 
-		// Compute any missing tangents. MikkTSpace should be use only when the user want to recompute the normals or tangents otherwise should always fallback on builtin
-		if (BuildData->BuildOptions.bUseMikkTSpace && (BuildData->BuildOptions.bComputeNormals || BuildData->BuildOptions.bComputeTangents))
-		{
-			Skeletal_ComputeTangents_MikkTSpace(BuildData, OverlappingCorners);
-		}
-		else
-		{
-			Skeletal_ComputeTangents(BuildData, OverlappingCorners);
-		}
+		// Compute any missing tangents. MikkTSpace should be use only when the user want to recompute the normals or tangents otherwise should always fallback on builtin tangent
+		Skeletal_ComputeTangents(BuildData, OverlappingCorners);
 
 		// At this point the mesh will have valid tangents.
 		BuildData->ValidateTangentArraySize();
@@ -4477,18 +4213,6 @@ bool FMeshUtilities::BuildSkeletalMesh(FSkeletalMeshLODModel& LODModel, const FR
 		}
 	};
 
-	// Temporarily supporting both import paths
-	if (!BuildOptions.bUseMikkTSpace)
-	{
-		bool bBuildSuccess = BuildSkeletalMesh_Legacy(LODModel, RefSkeleton, Influences, Wedges, Faces, Points, PointToOriginalMap, BuildOptions.OverlappingThresholds, BuildOptions.bComputeNormals, BuildOptions.bComputeTangents, BuildOptions.bComputeWeightedNormals, OutWarningMessages, OutWarningNames);
-		if (bBuildSuccess)
-		{
-			UpdateOverlappingVertices(LODModel);
-		}
-
-		return bBuildSuccess;
-	}
-
 	SkeletalMeshBuildData BuildData(
 		LODModel,
 		RefSkeleton,
diff --git a/Engine/Source/Developer/SkeletalMeshUtilitiesCommon/Private/LODUtilities.cpp b/Engine/Source/Developer/SkeletalMeshUtilitiesCommon/Private/LODUtilities.cpp
--- a/Engine/Source/Developer/SkeletalMeshUtilitiesCommon/Private/LODUtilities.cpp
+++ b/Engine/Source/Developer/SkeletalMeshUtilitiesCommon/Private/LODUtilities.cpp
@@ -202,14 +202,11 @@ bool FLODUtilities::RegenerateLOD(USkeletalMesh* SkeletalMesh, int32 NewLODCount
 {
 	if (SkeletalMesh)
 	{
+		FScopedSkeletalMeshPostEditChange ScopedPostEditChange(SkeletalMesh);
+
 		// Unbind any existing clothing assets before we regenerate all LODs
 		TArray<ClothingAssetUtils::FClothingAssetMeshBinding> ClothingBindings;
-		ClothingAssetUtils::GetMeshClothingAssetBindings(SkeletalMesh, ClothingBindings);
-
-		for (ClothingAssetUtils::FClothingAssetMeshBinding& Binding : ClothingBindings)
-		{
-			Binding.Asset->UnbindFromSkeletalMesh(SkeletalMesh, Binding.LODIndex);
-		}
+		FLODUtilities::UnbindClothingAndBackup(SkeletalMesh, ClothingBindings);
 
 		int32 LODCount = SkeletalMesh->GetLODNum();
 
@@ -239,7 +236,7 @@ bool FLODUtilities::RegenerateLOD(USkeletalMesh* SkeletalMesh, int32 NewLODCount
 			for (int32 LODIdx = CurrentNumLODs; LODIdx < LODCount; LODIdx++)
 			{
 				// if no previous setting found, it will use default setting. 
-				FLODUtilities::SimplifySkeletalMeshLOD(UpdateContext, LODIdx);
+				FLODUtilities::SimplifySkeletalMeshLOD(UpdateContext, LODIdx, false);
 			}
 		}
 		else
@@ -249,22 +246,13 @@ bool FLODUtilities::RegenerateLOD(USkeletalMesh* SkeletalMesh, int32 NewLODCount
 				FSkeletalMeshLODInfo& CurrentLODInfo = *(SkeletalMesh->GetLODInfo(LODIdx));
 				if ((bRegenerateEvenIfImported && LODIdx > 0) || (bGenerateBaseLOD && LODIdx == 0) || CurrentLODInfo.bHasBeenSimplified )
 				{
-					FLODUtilities::SimplifySkeletalMeshLOD(UpdateContext, LODIdx);
+					FLODUtilities::SimplifySkeletalMeshLOD(UpdateContext, LODIdx, false);
 				}
 			}
 		}
 
 		//Restore all clothing we can
-		for (ClothingAssetUtils::FClothingAssetMeshBinding& Binding : ClothingBindings)
-		{
-			if (SkeletalMesh->GetImportedModel()->LODModels.IsValidIndex(Binding.LODIndex) &&
-				SkeletalMesh->GetImportedModel()->LODModels[Binding.LODIndex].Sections.IsValidIndex(Binding.SectionIndex))
-			{
-				Binding.Asset->BindToSkeletalMesh(SkeletalMesh, Binding.LODIndex, Binding.SectionIndex, Binding.AssetInternalLodIndex);
-			}
-		}
-
-		SkeletalMesh->PostEditChange();
+		FLODUtilities::RestoreClothingFromBackup(SkeletalMesh, ClothingBindings);
 
 		return true;
 	}
@@ -296,21 +284,12 @@ void FLODUtilities::RemoveLOD(FSkeletalMeshUpdateContext& UpdateContext, int32 D
 	// If its a valid LOD, kill it.
 	if( DesiredLOD > 0 && DesiredLOD < SkelMeshModel->LODModels.Num() )
 	{
-		//We'll be modifying the skel mesh data so reregister
-
-		//TODO - do we need to reregister something else instead?
-		FMultiComponentReregisterContext ReregisterContext(UpdateContext.AssociatedComponents);
-
-		// Release rendering resources before deleting LOD
-		SkeletalMesh->ReleaseResources();
+		FScopedSkeletalMeshPostEditChange ScopedPostEditChange(SkeletalMesh);
 
 		// Block until this is done
-		FlushRenderingCommands();
 
 		SkelMeshModel->LODModels.RemoveAt(DesiredLOD);
 		SkeletalMesh->RemoveLODInfo(DesiredLOD);
-		SkeletalMesh->InitResources();
-
 		RefreshLODChange(SkeletalMesh);
 
 		// Set the forced LOD to Auto.
@@ -332,9 +311,6 @@ void FLODUtilities::RemoveLOD(FSkeletalMeshUpdateContext& UpdateContext, int32 D
 			}
 		}
 
-		// This will recache derived render data, and re-init resources
-		SkeletalMesh->PostEditChange();
-
 		//Notify calling system of change
 		UpdateContext.OnLODChanged.ExecuteIfBound();
 
@@ -488,7 +464,7 @@ void ProjectTargetOnBase(const TArray<FSoftSkinVertex>& BaseVertices, const TArr
 	{
 		//Use the remap base index in case some sections disappear during the reduce phase
 		int32 BaseSectionIndex = TargetSectionMatchBaseIndex[SectionIndex];
-		if (BaseSectionIndex == INDEX_NONE || !PerSectionBaseTriangleIndices.IsValidIndex(BaseSectionIndex))
+		if (BaseSectionIndex == INDEX_NONE || !PerSectionBaseTriangleIndices.IsValidIndex(BaseSectionIndex) || PerSectionBaseTriangleIndices[BaseSectionIndex].Num() < 1)
 		{
 			continue;
 		}
@@ -518,6 +494,7 @@ void ProjectTargetOnBase(const TArray<FSoftSkinVertex>& BaseVertices, const TArr
 			TriangleElement.TriangleIndex = Triangles.Num();
 			Triangles.Add(TriangleElement);
 		}
+		check(!BaseMeshUVBound.GetExtent().IsNearlyZero());
 		//Setup the Quad tree
 		float UVsQuadTreeMinSize = 0.001f;
 		TQuadTree<uint32, 100> QuadTree(BaseMeshUVBound, UVsQuadTreeMinSize);
@@ -661,7 +638,7 @@ void CreateLODMorphTarget(USkeletalMesh* SkeletalMesh, FReductionBaseSkeletalMes
 	if (ReductionBaseSkeletalMeshBulkData != nullptr)
 	{
 		FSkeletalMeshLODModel TempBaseLODModel;
-		ReductionBaseSkeletalMeshBulkData->LoadReductionData(TempBaseLODModel, BaseLODMorphTargetData);
+		ReductionBaseSkeletalMeshBulkData->LoadReductionData(TempBaseLODModel, BaseLODMorphTargetData, SkeletalMesh);
 	}
 
 	FSkeletalMeshModel* SkeletalMeshModel = SkeletalMesh->GetImportedModel();
@@ -841,7 +818,8 @@ void FLODUtilities::ApplyMorphTargetsToLOD(USkeletalMesh* SkeletalMesh, int32 So
 	TMap<FString, TArray<FMorphTargetDelta>> TempBaseLODMorphTargetData;
 	if (bReduceBaseLOD)
 	{
-		ReductionBaseSkeletalMeshBulkData->LoadReductionData(TempBaseLODModel, TempBaseLODMorphTargetData);
+		check(ReductionBaseSkeletalMeshBulkData != nullptr);
+		ReductionBaseSkeletalMeshBulkData->LoadReductionData(TempBaseLODModel, TempBaseLODMorphTargetData, SkeletalMesh);
 	}
 
 	const FSkeletalMeshLODModel& BaseLODModel = bReduceBaseLOD ? TempBaseLODModel : SkeletalMeshResource->LODModels[SourceLOD];
@@ -900,7 +878,7 @@ void FLODUtilities::ApplyMorphTargetsToLOD(USkeletalMesh* SkeletalMesh, int32 So
 	for (int32 SectionIndex = 0; SectionIndex < SectionCount; ++SectionIndex)
 	{
 		const FSkelMeshSection& Section = BaseLODModel.Sections[SectionIndex];
-		if (Section.HasClothingData())
+		if (Section.ClothingData.AssetGuid.IsValid())
 		{
 			continue;
 		}
@@ -964,13 +942,14 @@ void FLODUtilities::ApplyMorphTargetsToLOD(USkeletalMesh* SkeletalMesh, int32 So
 	CreateLODMorphTarget(SkeletalMesh, ReductionBaseSkeletalMeshBulkData, SourceLOD, DestinationLOD, PerMorphTargetBaseIndexToMorphTargetDelta, BaseMorphIndexToTargetIndexList, TargetVertices, TargetMatchData);
 }
 
-void FLODUtilities::SimplifySkeletalMeshLOD( USkeletalMesh* SkeletalMesh, int32 DesiredLOD, bool bReregisterComponent /*= true*/, bool bRestoreClothing /*= false*/)
+void FLODUtilities::SimplifySkeletalMeshLOD( USkeletalMesh* SkeletalMesh, int32 DesiredLOD, bool bRestoreClothing /*= false*/)
 {
 	IMeshReductionModule& ReductionModule = FModuleManager::Get().LoadModuleChecked<IMeshReductionModule>("MeshReductionInterface");
 	IMeshReduction* MeshReduction = ReductionModule.GetSkeletalMeshReductionInterface();
 
 	check (MeshReduction && MeshReduction->IsSupported());
 
+
 	if (DesiredLOD == 0
 		&& SkeletalMesh->GetLODInfo(DesiredLOD) != nullptr
 		&& SkeletalMesh->GetLODInfo(DesiredLOD)->bHasBeenSimplified
@@ -1000,18 +979,13 @@ void FLODUtilities::SimplifySkeletalMeshLOD( USkeletalMesh* SkeletalMesh, int32
 		GWarn->BeginSlowTask(StatusUpdate, true);
 	}
 
+	FScopedSkeletalMeshPostEditChange ScopedPostEditChange(SkeletalMesh);
+
 	// Unbind DesiredLOD existing clothing assets before we simplify this LOD
 	TArray<ClothingAssetUtils::FClothingAssetMeshBinding> ClothingBindings;
-	if (bRestoreClothing)
+	if (bRestoreClothing && SkeletalMesh->GetImportedModel() && SkeletalMesh->GetImportedModel()->LODModels.IsValidIndex(DesiredLOD))
 	{
-		ClothingAssetUtils::GetMeshClothingAssetBindings(SkeletalMesh, ClothingBindings);
-		for (ClothingAssetUtils::FClothingAssetMeshBinding& Binding : ClothingBindings)
-		{
-			if (DesiredLOD == Binding.LODIndex)
-			{
-				Binding.Asset->UnbindFromSkeletalMesh(SkeletalMesh, Binding.LODIndex);
-			}
-		}
+		FLODUtilities::UnbindClothingAndBackup(SkeletalMesh, ClothingBindings, DesiredLOD);
 	}
 
 	if (SkeletalMesh->GetLODInfo(DesiredLOD) != nullptr)
@@ -1049,7 +1023,7 @@ void FLODUtilities::SimplifySkeletalMeshLOD( USkeletalMesh* SkeletalMesh, int32
 					}
 				}
 				//Copy the original SkeletalMesh LODModel
-				SkeletalMeshResource->OriginalReductionSourceMeshData[DesiredLOD]->SaveReductionData(SrcModel, BaseLODMorphTargetData);
+				SkeletalMeshResource->OriginalReductionSourceMeshData[DesiredLOD]->SaveReductionData(SrcModel, BaseLODMorphTargetData, SkeletalMesh);
 
 				if (DesiredLOD == 0)
 				{
@@ -1060,7 +1034,7 @@ void FLODUtilities::SimplifySkeletalMeshLOD( USkeletalMesh* SkeletalMesh, int32
 	}
 	
 
-	if (MeshReduction->ReduceSkeletalMesh(SkeletalMesh, DesiredLOD, bReregisterComponent))
+	if (MeshReduction->ReduceSkeletalMesh(SkeletalMesh, DesiredLOD))
 	{
 		check(SkeletalMesh->GetLODNum() >= 1);
 
@@ -1078,21 +1052,7 @@ void FLODUtilities::SimplifySkeletalMeshLOD( USkeletalMesh* SkeletalMesh, int32
 			}
 		};
 
-		if (bReregisterComponent)
-		{
-			TComponentReregisterContext<USkinnedMeshComponent> ReregisterContext;
-			SkeletalMesh->ReleaseResources();
-			SkeletalMesh->ReleaseResourcesFence.Wait();
-
-			ApplyMorphTargetOption();
-			
-			SkeletalMesh->PostEditChange();
-			SkeletalMesh->InitResources();
-		}
-		else
-		{
-			ApplyMorphTargetOption();
-		}
+		ApplyMorphTargetOption();
 		SkeletalMesh->MarkPackageDirty();
 	}
 	else
@@ -1105,25 +1065,15 @@ void FLODUtilities::SimplifySkeletalMeshLOD( USkeletalMesh* SkeletalMesh, int32
 	}
 
 	//Put back the clothing for the DesiredLOD
-	if (bRestoreClothing)
+	if (bRestoreClothing && ClothingBindings.Num() > 0 && SkeletalMesh->GetImportedModel()->LODModels.IsValidIndex(DesiredLOD))
 	{
-		for (ClothingAssetUtils::FClothingAssetMeshBinding& Binding : ClothingBindings)
-		{
-			if (SkeletalMesh->GetImportedModel()->LODModels.IsValidIndex(Binding.LODIndex) &&
-				SkeletalMesh->GetImportedModel()->LODModels[Binding.LODIndex].Sections.IsValidIndex(Binding.SectionIndex))
-			{
-				if (DesiredLOD == Binding.LODIndex)
-				{
-					Binding.Asset->BindToSkeletalMesh(SkeletalMesh, Binding.LODIndex, Binding.SectionIndex, Binding.AssetInternalLodIndex);
-				}
-			}
-		}
+		FLODUtilities::RestoreClothingFromBackup(SkeletalMesh, ClothingBindings, DesiredLOD);
 	}
 
 	GWarn->EndSlowTask();
 }
 
-void FLODUtilities::SimplifySkeletalMeshLOD(FSkeletalMeshUpdateContext& UpdateContext, int32 DesiredLOD, bool bReregisterComponent /*= true*/, bool bRestoreClothing /*= false*/)
+void FLODUtilities::SimplifySkeletalMeshLOD(FSkeletalMeshUpdateContext& UpdateContext, int32 DesiredLOD, bool bRestoreClothing /*= false*/)
 {
 	USkeletalMesh* SkeletalMesh = UpdateContext.SkeletalMesh;
 	IMeshReductionModule& ReductionModule = FModuleManager::Get().LoadModuleChecked<IMeshReductionModule>("MeshReductionInterface");
@@ -1131,7 +1081,7 @@ void FLODUtilities::SimplifySkeletalMeshLOD(FSkeletalMeshUpdateContext& UpdateCo
 
 	if (MeshReduction && MeshReduction->IsSupported() && SkeletalMesh)
 	{
-		SimplifySkeletalMeshLOD(SkeletalMesh, DesiredLOD, bReregisterComponent, bRestoreClothing);
+		SimplifySkeletalMeshLOD(SkeletalMesh, DesiredLOD, bRestoreClothing);
 		
 		if (UpdateContext.OnLODChanged.IsBound())
 		{
@@ -1141,40 +1091,27 @@ void FLODUtilities::SimplifySkeletalMeshLOD(FSkeletalMeshUpdateContext& UpdateCo
 	}
 }
 
-void FLODUtilities::RestoreSkeletalMeshLODImportedData(USkeletalMesh* SkeletalMesh, int32 LodIndex, bool bReregisterComponent /*= true*/)
+void FLODUtilities::RestoreSkeletalMeshLODImportedData(USkeletalMesh* SkeletalMesh, int32 LodIndex)
 {
-
 	if (!SkeletalMesh->GetImportedModel()->OriginalReductionSourceMeshData.IsValidIndex(LodIndex) || SkeletalMesh->GetImportedModel()->OriginalReductionSourceMeshData[LodIndex]->IsEmpty())
 	{
 		//There is nothing to restore
 		return;
 	}
 
+	FScopedSkeletalMeshPostEditChange ScopedPostEditChange(SkeletalMesh);
+
 	// Unbind LodIndex existing clothing assets before restoring the LOD
 	TArray<ClothingAssetUtils::FClothingAssetMeshBinding> ClothingBindings;
-	ClothingAssetUtils::GetMeshClothingAssetBindings(SkeletalMesh, ClothingBindings);
-	for (ClothingAssetUtils::FClothingAssetMeshBinding& Binding : ClothingBindings)
-	{
-		//Unbind only the LOD we restore
-		if (Binding.LODIndex == LodIndex)
-		{
-			Binding.Asset->UnbindFromSkeletalMesh(SkeletalMesh, Binding.LODIndex);
-		}
-	}
+	FLODUtilities::UnbindClothingAndBackup(SkeletalMesh, ClothingBindings);
 
 	FSkeletalMeshLODModel ImportedBaseLODModel;
 	TMap<FString, TArray<FMorphTargetDelta>> ImportedBaseLODMorphTargetData;
-	SkeletalMesh->GetImportedModel()->OriginalReductionSourceMeshData[LodIndex]->LoadReductionData(ImportedBaseLODModel, ImportedBaseLODMorphTargetData);
+	SkeletalMesh->GetImportedModel()->OriginalReductionSourceMeshData[LodIndex]->LoadReductionData(ImportedBaseLODModel, ImportedBaseLODMorphTargetData, SkeletalMesh);
 	{
 		FSkeletalMeshUpdateContext UpdateContext;
 		UpdateContext.SkeletalMesh = SkeletalMesh;
 
-		TComponentReregisterContext<USkinnedMeshComponent> ReregisterContext;
-		if (bReregisterComponent)
-		{
-			SkeletalMesh->ReleaseResources();
-			SkeletalMesh->ReleaseResourcesFence.Wait();
-		}
 		//Copy the SkeletalMeshLODModel
 		SkeletalMesh->GetImportedModel()->LODModels[LodIndex] = ImportedBaseLODModel;
 		//Copy the morph target deltas
@@ -1196,20 +1133,7 @@ void FLODUtilities::RestoreSkeletalMeshLODImportedData(USkeletalMesh* SkeletalMe
 		SkeletalMesh->GetImportedModel()->OriginalReductionSourceMeshData[LodIndex]->EmptyBulkData();
 
 		//Put back the clothing for the restore LOD
-		for (ClothingAssetUtils::FClothingAssetMeshBinding& Binding : ClothingBindings)
-		{
-			if (LodIndex == Binding.LODIndex && SkeletalMesh->GetImportedModel()->LODModels.IsValidIndex(Binding.LODIndex) &&
-				SkeletalMesh->GetImportedModel()->LODModels[Binding.LODIndex].Sections.IsValidIndex(Binding.SectionIndex))
-			{
-				Binding.Asset->BindToSkeletalMesh(SkeletalMesh, Binding.LODIndex, Binding.SectionIndex, Binding.AssetInternalLodIndex);
-			}
-		}
-
-		if (bReregisterComponent)
-		{
-			SkeletalMesh->PostEditChange();
-			SkeletalMesh->InitResources();
-		}
+		FLODUtilities::RestoreClothingFromBackup(SkeletalMesh, ClothingBindings);
 
 		if (UpdateContext.OnLODChanged.IsBound())
 		{
@@ -1600,17 +1524,21 @@ void MatchVertexIndexUsingPosition(
 
 bool FLODUtilities::UpdateAlternateSkinWeights(USkeletalMesh* SkeletalMeshDest, const FName& ProfileNameDest, int32 LODIndexDest, FOverlappingThresholds OverlappingThresholds, bool ShouldImportNormals, bool ShouldImportTangents, bool bUseMikkTSpace, bool bComputeWeightedNormals)
 {
-	//Ensure log message only once
-	bool bNoMatchMsgDone = false;
-
 	//Grab all the destination structure
 	check(SkeletalMeshDest);
 	check(SkeletalMeshDest->GetImportedModel());
 	check(SkeletalMeshDest->GetImportedModel()->LODModels.IsValidIndex(LODIndexDest));
 	FSkeletalMeshLODModel& LODModelDest = SkeletalMeshDest->GetImportedModel()->LODModels[LODIndexDest];
+	return UpdateAlternateSkinWeights(LODModelDest, SkeletalMeshDest->GetName(), SkeletalMeshDest->RefSkeleton, ProfileNameDest, LODIndexDest, OverlappingThresholds, ShouldImportNormals, ShouldImportTangents, bUseMikkTSpace, bComputeWeightedNormals);
+}
+
+bool FLODUtilities::UpdateAlternateSkinWeights(FSkeletalMeshLODModel& LODModelDest, const FString SkeletalMeshName, FReferenceSkeleton& RefSkeleton, const FName& ProfileNameDest, int32 LODIndexDest, FOverlappingThresholds OverlappingThresholds, bool ShouldImportNormals, bool ShouldImportTangents, bool bUseMikkTSpace, bool bComputeWeightedNormals)
+{
+	//Ensure log message only once
+	bool bNoMatchMsgDone = false;
 	if (LODModelDest.RawSkeletalMeshBulkData.IsEmpty())
 	{
-		UE_LOG(LogLODUtilities, Error, TEXT("Failed to import Skin Weight Profile as the target skeletal mesh (%s) requires reimporting first."), SkeletalMeshDest ? *SkeletalMeshDest->GetName() : TEXT("NULL"));
+		UE_LOG(LogLODUtilities, Error, TEXT("Failed to import Skin Weight Profile as the target skeletal mesh (%s) requires reimporting first."), *SkeletalMeshName);
 		//Very old asset will not have this data, we cannot add alternate until the asset is reimported
 		return false;
 	}
@@ -1622,7 +1550,7 @@ bool FLODUtilities::UpdateAlternateSkinWeights(USkeletalMesh* SkeletalMeshDest,
 	int32 ProfileIndex = 0;
 	if (!ImportDataDest.AlternateInfluenceProfileNames.Find(ProfileNameDest.ToString(), ProfileIndex))
 	{
-		UE_LOG(LogLODUtilities, Error, TEXT("Failed to import Skin Weight Profile the alternate skinning imported source data is not available."), SkeletalMeshDest ? *SkeletalMeshDest->GetName() : TEXT("NULL"));
+		UE_LOG(LogLODUtilities, Error, TEXT("Failed to import Skin Weight Profile the alternate skinning imported source data is not available."), *SkeletalMeshName);
 		return false;
 	}
 	check(ImportDataDest.AlternateInfluences.IsValidIndex(ProfileIndex));
@@ -1634,7 +1562,7 @@ bool FLODUtilities::UpdateAlternateSkinWeights(USkeletalMesh* SkeletalMeshDest,
 
 	if (ImportDataDest.NumTexCoords <= 0 || ImportDataSrc.NumTexCoords <= 0)
 	{
-		UE_LOG(LogLODUtilities, Error, TEXT("Failed to import Skin Weight Profile as the target skeletal mesh (%s) or imported file does not contain UV coordinates."), SkeletalMeshDest ? *SkeletalMeshDest->GetName() : TEXT("NULL"));
+		UE_LOG(LogLODUtilities, Error, TEXT("Failed to import Skin Weight Profile as the target skeletal mesh (%s) or imported file does not contain UV coordinates."), *SkeletalMeshName);
 		return false;
 	}
 
@@ -1650,7 +1578,7 @@ bool FLODUtilities::UpdateAlternateSkinWeights(USkeletalMesh* SkeletalMeshDest,
 
 	//Create a map to remap source bone index to destination bone index
 	TMap<int32, int32> RemapBoneIndexSrcToDest;
-	FillRemapBoneIndexSrcToDest(ImportDataSrc, ImportDataDest, SkeletalMeshDest->GetName(), LODIndexDest, RemapBoneIndexSrcToDest);
+	FillRemapBoneIndexSrcToDest(ImportDataSrc, ImportDataDest, SkeletalMeshName, LODIndexDest, RemapBoneIndexSrcToDest);
 
 	//Map to get the vertex index source to a destination vertex match
 	TSortedMap<uint32, VertexMatchNameSpace::FVertexMatchResult> VertexIndexSrcToVertexIndexDestMatches;
@@ -1879,7 +1807,10 @@ bool FLODUtilities::UpdateAlternateSkinWeights(USkeletalMesh* SkeletalMeshDest,
 	TArray<FText> WarningMessages;
 	TArray<FName> WarningNames;
 	//Build the destination mesh with the Alternate influences, so the chunking is done properly.
-	bool bBuildSuccess = MeshUtilities.BuildSkeletalMesh(LODModelDest, SkeletalMeshDest->RefSkeleton, LODInfluencesDest, LODWedgesDest, LODFacesDest, LODPointsDest, LODPointToRawMapDest, BuildOptions, &WarningMessages, &WarningNames);
+	bool bBuildSuccess = MeshUtilities.BuildSkeletalMesh(LODModelDest, RefSkeleton, LODInfluencesDest, LODWedgesDest, LODFacesDest, LODPointsDest, LODPointToRawMapDest, BuildOptions, &WarningMessages, &WarningNames);
+	//Re-Apply the user section changes, the UserSectionsData is map to original section and should match the builded LODModel
+	LODModelDest.SyncronizeUserSectionsDataArray();
+
 	RegenerateAllImportSkinWeightProfileData(LODModelDest);
 	
 	return bBuildSuccess;
@@ -2004,7 +1935,8 @@ void FLODUtilities::RegenerateDependentLODs(USkeletalMesh* SkeletalMesh, int32 L
 {
 	FSkeletalMeshUpdateContext UpdateContext;
 	UpdateContext.SkeletalMesh = SkeletalMesh;
-	//Check the dependencies and regenerate the LODs acoording to it
+
+	//Check the dependencies and regenerate the LODs according to it
 	TArray<bool> LODDependencies;
 	int32 LODNumber = SkeletalMesh->GetLODNum();
 	LODDependencies.AddZeroed(LODNumber);
@@ -2022,19 +1954,20 @@ void FLODUtilities::RegenerateDependentLODs(USkeletalMesh* SkeletalMesh, int32 L
 	}
 	if (bRegenLODs)
 	{
-		TComponentReregisterContext<USkinnedMeshComponent> ReregisterContext;
-		SkeletalMesh->Modify();
-		SkeletalMesh->ReleaseResources();
-		SkeletalMesh->ReleaseResourcesFence.Wait();
-		for (int32 DependentLODIndex = LODIndex + 1; DependentLODIndex < LODNumber; ++DependentLODIndex)
+		FScopedSkeletalMeshPostEditChange ScopedPostEditChange(SkeletalMesh);
+		auto SimplifyDependentLODS = [&LODIndex, &LODNumber, &LODDependencies, &UpdateContext]()
 		{
-			if (LODDependencies[DependentLODIndex])
+			for (int32 DependentLODIndex = LODIndex + 1; DependentLODIndex < LODNumber; ++DependentLODIndex)
 			{
-				FLODUtilities::SimplifySkeletalMeshLOD(UpdateContext, DependentLODIndex, false);
+				if (LODDependencies[DependentLODIndex])
+				{
+					FLODUtilities::SimplifySkeletalMeshLOD(UpdateContext, DependentLODIndex, true);
+				}
 			}
-		}
-		SkeletalMesh->PostEditChange();
-		SkeletalMesh->InitResources();
+		};
+
+		SkeletalMesh->Modify();
+		SimplifyDependentLODS();
 	}
 }
 
@@ -2486,4 +2419,86 @@ void FLODUtilities::BuildMorphTargets(USkeletalMesh* BaseSkelMesh, FSkeletalMesh
 	}
 }
 
+void FLODUtilities::UnbindClothingAndBackup(USkeletalMesh* SkeletalMesh, TArray<ClothingAssetUtils::FClothingAssetMeshBinding>& ClothingBindings)
+{
+	for (int32 LODIndex = 0; LODIndex < SkeletalMesh->GetImportedModel()->LODModels.Num(); ++LODIndex)
+	{
+		TArray<ClothingAssetUtils::FClothingAssetMeshBinding> LODBindings;
+		UnbindClothingAndBackup(SkeletalMesh, LODBindings, LODIndex);
+		ClothingBindings.Append(LODBindings);
+	}
+}
+
+void FLODUtilities::UnbindClothingAndBackup(USkeletalMesh* SkeletalMesh, TArray<ClothingAssetUtils::FClothingAssetMeshBinding>& ClothingBindings, const int32 LODIndex)
+{
+	if (!SkeletalMesh->GetImportedModel()->LODModels.IsValidIndex(LODIndex))
+	{
+		return;
+	}
+	FSkeletalMeshLODModel& LODModel = SkeletalMesh->GetImportedModel()->LODModels[LODIndex];
+	//Store the clothBinding
+	ClothingAssetUtils::GetMeshClothingAssetBindings(SkeletalMesh, ClothingBindings, LODIndex);
+	//Unbind the Cloth for this LOD before we reduce it, we will put back the cloth after the reduction, if it still match the sections
+	for (ClothingAssetUtils::FClothingAssetMeshBinding& Binding : ClothingBindings)
+	{
+		if (Binding.LODIndex == LODIndex)
+		{
+			check(Binding.Asset);
+			Binding.Asset->UnbindFromSkeletalMesh(SkeletalMesh, Binding.LODIndex);
+			//Use the UserSectionsData original section index, this will ensure we remap correctly the cloth if the reduction has change the number of sections
+			int32 OriginalDataSectionIndex = LODModel.Sections[Binding.SectionIndex].OriginalDataSectionIndex;
+			Binding.SectionIndex = OriginalDataSectionIndex;
+		}
+	}
+}
+
+void FLODUtilities::RestoreClothingFromBackup(USkeletalMesh* SkeletalMesh, TArray<ClothingAssetUtils::FClothingAssetMeshBinding>& ClothingBindings, FOnPostRebindCloth OnPostRebindCloth)
+{
+	for (int32 LODIndex = 0; LODIndex < SkeletalMesh->GetImportedModel()->LODModels.Num(); ++LODIndex)
+	{
+		RestoreClothingFromBackup(SkeletalMesh, ClothingBindings, LODIndex, OnPostRebindCloth);
+	}
+}
+
+void FLODUtilities::RestoreClothingFromBackup(USkeletalMesh* SkeletalMesh, TArray<ClothingAssetUtils::FClothingAssetMeshBinding>& ClothingBindings, const int32 LODIndex, FOnPostRebindCloth OnPostRebindCloth)
+{
+	if (!SkeletalMesh->GetImportedModel()->LODModels.IsValidIndex(LODIndex))
+	{
+		return;
+	}
+	FSkeletalMeshLODModel& LODModel = SkeletalMesh->GetImportedModel()->LODModels[LODIndex];
+	for (ClothingAssetUtils::FClothingAssetMeshBinding& Binding : ClothingBindings)
+	{
+		for (int32 SectionIndex = 0; SectionIndex < LODModel.Sections.Num(); ++SectionIndex)
+		{
+			if (LODModel.Sections[SectionIndex].OriginalDataSectionIndex != Binding.SectionIndex)
+			{
+				continue;
+			}
+			if (Binding.LODIndex == LODIndex)
+			{
+				check(Binding.Asset);
+				if (Binding.Asset->BindToSkeletalMesh(SkeletalMesh, Binding.LODIndex, SectionIndex, Binding.AssetInternalLodIndex))
+				{
+					OnPostRebindCloth.ExecuteIfBound(Binding.LODIndex, LODModel.Sections[SectionIndex]);
+				}
+			}
+			break;
+		}
+	}
+}
+
+void FLODUtilities::RestoreClothingFromBackup(USkeletalMesh* SkeletalMesh, TArray<ClothingAssetUtils::FClothingAssetMeshBinding>& ClothingBindings)
+{
+	FOnPostRebindCloth EmptyDelegate;
+	RestoreClothingFromBackup(SkeletalMesh, ClothingBindings, EmptyDelegate);
+}
+
+void FLODUtilities::RestoreClothingFromBackup(USkeletalMesh* SkeletalMesh, TArray<ClothingAssetUtils::FClothingAssetMeshBinding>& ClothingBindings, const int32 LODIndex)
+{
+	FOnPostRebindCloth EmptyDelegate;
+	RestoreClothingFromBackup(SkeletalMesh, ClothingBindings, LODIndex, EmptyDelegate);
+}
+
+
 #undef LOCTEXT_NAMESPACE // "LODUtilities"
\ No newline at end of file
diff --git a/Engine/Source/Developer/SkeletalMeshUtilitiesCommon/Public/LODUtilities.h b/Engine/Source/Developer/SkeletalMeshUtilitiesCommon/Public/LODUtilities.h
--- a/Engine/Source/Developer/SkeletalMeshUtilitiesCommon/Public/LODUtilities.h
+++ b/Engine/Source/Developer/SkeletalMeshUtilitiesCommon/Public/LODUtilities.h
@@ -8,7 +8,10 @@
 #include "Rendering/SkeletalMeshLODImporterData.h"
 #include "Framework/Commands/UIAction.h"
 
-
+namespace ClothingAssetUtils
+{
+	struct FClothingAssetMeshBinding;
+}
 
 //////////////////////////////////////////////////////////////////////////
 // FSkeletalMeshUpdateContext
@@ -24,6 +27,7 @@ struct FSkeletalMeshUpdateContext
 
 //////////////////////////////////////////////////////////////////////////
 // FLODUtilities
+DECLARE_DELEGATE_TwoParams(FOnPostRebindCloth, int32, const FSkelMeshSection&);
 
 class SKELETALMESHUTILITIESCOMMON_API FLODUtilities
 {
@@ -62,7 +66,7 @@ public:
 	* @param UpdateContext - The skeletal mesh and actor components to operate on.
 	* @param DesiredLOD - The LOD to simplify
 	*/
-	static void SimplifySkeletalMeshLOD(FSkeletalMeshUpdateContext& UpdateContext, int32 DesiredLOD, bool bReregisterComponent = true, bool bRestoreClothing = false);
+	static void SimplifySkeletalMeshLOD(FSkeletalMeshUpdateContext& UpdateContext, int32 DesiredLOD, bool bRestoreClothing = false);
 
 	/**
 	*	Restore the LOD imported model to the last imported data. Call this function if you want to remove the reduce on the base LOD
@@ -71,7 +75,7 @@ public:
 	* @param LodIndex - The LOD index to restore the imported LOD model
 	* @param bReregisterComponent - if true the component using the skeletal mesh will all be re register.
 	*/
-	static void RestoreSkeletalMeshLODImportedData(USkeletalMesh* SkeletalMesh, int32 LodIndex, bool bReregisterComponent = true);
+	static void RestoreSkeletalMeshLODImportedData(USkeletalMesh* SkeletalMesh, int32 LodIndex);
 	
 	/**
 	 * Refresh LOD Change
@@ -117,9 +121,20 @@ public:
 	 */
 	static bool UpdateAlternateSkinWeights(USkeletalMesh* SkeletalMeshDest, const FName& ProfileNameDest, int32 LODIndexDest, FOverlappingThresholds OverlappingThresholds, bool ShouldImportNormals, bool ShouldImportTangents, bool bUseMikkTSpace, bool bComputeWeightedNormals);
 
+	static bool UpdateAlternateSkinWeights(FSkeletalMeshLODModel& LODModelDest, const FString SkeletalMeshName, FReferenceSkeleton& RefSkeleton, const FName& ProfileNameDest, int32 LODIndexDest, FOverlappingThresholds OverlappingThresholds, bool ShouldImportNormals, bool ShouldImportTangents, bool bUseMikkTSpace, bool bComputeWeightedNormals);
+
 	/** Re-generate all (editor-only) skin weight profile, used whenever we rebuild the skeletal mesh data which could change the chunking and bone indices */
 	static void RegenerateAllImportSkinWeightProfileData(FSkeletalMeshLODModel& LODModelDest);
 
+	static void UnbindClothingAndBackup(USkeletalMesh* SkeletalMesh, TArray<ClothingAssetUtils::FClothingAssetMeshBinding>& ClothingBindings);
+	static void UnbindClothingAndBackup(USkeletalMesh* SkeletalMesh, TArray<ClothingAssetUtils::FClothingAssetMeshBinding>& ClothingBindings, const int32 LODIndex);
+	
+	static void RestoreClothingFromBackup(USkeletalMesh* SkeletalMesh, TArray<ClothingAssetUtils::FClothingAssetMeshBinding>& ClothingBindings);
+	static void RestoreClothingFromBackup(USkeletalMesh* SkeletalMesh, TArray<ClothingAssetUtils::FClothingAssetMeshBinding>& ClothingBindings, const int32 LODIndex);
+	static void RestoreClothingFromBackup(USkeletalMesh* SkeletalMesh, TArray<ClothingAssetUtils::FClothingAssetMeshBinding>& ClothingBindings, FOnPostRebindCloth OnPostRebindCloth);
+	static void RestoreClothingFromBackup(USkeletalMesh* SkeletalMesh, TArray<ClothingAssetUtils::FClothingAssetMeshBinding>& ClothingBindings, const int32 LODIndex, FOnPostRebindCloth OnPostRebindCloth);
+	
+
 private:
 	FLODUtilities() {}
 
@@ -134,7 +149,7 @@ private:
 	 * @param SkeletalMesh - The skeletal mesh and actor components to operate on.
 	 * @param DesiredLOD - Desired LOD
 	 */
-	static void SimplifySkeletalMeshLOD(USkeletalMesh* SkeletalMesh, int32 DesiredLOD, bool bReregisterComponent = true, bool bRestoreClothing = false);
+	static void SimplifySkeletalMeshLOD(USkeletalMesh* SkeletalMesh, int32 DesiredLOD, bool bRestoreClothing = false);
 
 	/**
 	*  Remap the morph targets of the base LOD onto the desired LOD.
diff --git a/Engine/Source/Editor/PropertyEditor/Private/PropertyCustomizationHelpers.cpp b/Engine/Source/Editor/PropertyEditor/Private/PropertyCustomizationHelpers.cpp
--- a/Engine/Source/Editor/PropertyEditor/Private/PropertyCustomizationHelpers.cpp
+++ b/Engine/Source/Editor/PropertyEditor/Private/PropertyCustomizationHelpers.cpp
@@ -815,9 +815,9 @@ public:
 	* @param Section		The Section being used
 	* @param bCanBeReplced	Whether or not the Section can be replaced by a user
 	*/
-	virtual void AddSection(int32 LodIndex, int32 SectionIndex, FName InMaterialSlotName, int32 InMaterialSlotIndex, FName InOriginalMaterialSlotName, const TMap<int32, FName> &InAvailableMaterialSlotName, const UMaterialInterface* Material, bool IsSectionUsingCloth) override
+	virtual void AddSection(int32 LodIndex, int32 SectionIndex, FName InMaterialSlotName, int32 InMaterialSlotIndex, FName InOriginalMaterialSlotName, const TMap<int32, FName> &InAvailableMaterialSlotName, const UMaterialInterface* Material, bool IsSectionUsingCloth, bool bIsChunkSection) override
 	{
-		FSectionListItem SectionItem(LodIndex, SectionIndex, InMaterialSlotName, InMaterialSlotIndex, InOriginalMaterialSlotName, InAvailableMaterialSlotName, Material, IsSectionUsingCloth, ThumbnailSize);
+		FSectionListItem SectionItem(LodIndex, SectionIndex, InMaterialSlotName, InMaterialSlotIndex, InOriginalMaterialSlotName, InAvailableMaterialSlotName, Material, IsSectionUsingCloth, ThumbnailSize, bIsChunkSection);
 		if (!Sections.Contains(SectionItem))
 		{
 			Sections.Add(SectionItem);
@@ -948,6 +948,7 @@ public:
 				.HAlign(HAlign_Fill)
 				[
 					SNew(SHorizontalBox)
+					.Visibility(SectionItem.bIsChunkSection ? EVisibility::Collapsed : EVisibility::Visible)
 					+ SHorizontalBox::Slot()
 					.FillWidth(1.0f)
 					[
@@ -1023,6 +1024,22 @@ public:
 						]
 					]
 				]
+				+ SVerticalBox::Slot()
+				.AutoHeight()
+				.Padding(0.0f)
+				.VAlign(VAlign_Center)
+				.HAlign(HAlign_Fill)
+				[
+					SNew(SHorizontalBox)
+					.Visibility(SectionItem.bIsChunkSection ? EVisibility::Visible : EVisibility::Collapsed)
+					+ SHorizontalBox::Slot()
+					.FillWidth(1.0f)
+					[
+						SNew(STextBlock)
+						.Font(IDetailLayoutBuilder::GetDetailFont())
+						.Text(LOCTEXT("SectionListItemChunkSectionValueLabel", "Chunked"))
+					]
+				]
 			];
 	}
 
@@ -1210,7 +1227,7 @@ void FSectionList::GenerateChildContent(IDetailChildrenBuilder& ChildrenBuilder)
 			if (bDisplayAllSectionsInSlot)
 			{
 				FDetailWidgetRow& ChildRow = ChildrenBuilder.AddCustomRow(Section.Material.IsValid() ? FText::FromString(Section.Material->GetName()) : FText::GetEmpty());
-				AddSectionItem(ChildRow, CurrentLODIndex, FSectionListItem(CurrentLODIndex, Section.SectionIndex, Section.MaterialSlotName, Section.MaterialSlotIndex, Section.OriginalMaterialSlotName, Section.AvailableMaterialSlotName, Section.Material.Get(), Section.IsSectionUsingCloth, ThumbnailSize), !bDisplayAllSectionsInSlot);
+				AddSectionItem(ChildRow, CurrentLODIndex, FSectionListItem(CurrentLODIndex, Section.SectionIndex, Section.MaterialSlotName, Section.MaterialSlotIndex, Section.OriginalMaterialSlotName, Section.AvailableMaterialSlotName, Section.Material.Get(), Section.IsSectionUsingCloth, ThumbnailSize, Section.bIsChunkSection), !bDisplayAllSectionsInSlot);
 			}
 		}
 	}
@@ -1292,6 +1309,7 @@ void FSectionList::AddSectionItem(FDetailWidgetRow& Row, int32 LodIndex, const s
 {
 	uint32 NumSections = SectionListBuilder->GetNumSections(LodIndex);
 
+	bool bIsChunkSection = Item.bIsChunkSection;
 	TSharedRef<FSectionItemView> NewView = FSectionItemView::Create(Item, SectionListDelegates.OnSectionChanged, SectionListDelegates.OnGenerateCustomNameWidgets, SectionListDelegates.OnGenerateCustomSectionWidgets, SectionListDelegates.OnResetSectionToDefaultClicked, NumSections, ThumbnailSize);
 
 	TSharedPtr<SWidget> RightSideContent;
@@ -1318,13 +1336,17 @@ void FSectionList::AddSectionItem(FDetailWidgetRow& Row, int32 LodIndex, const s
 		ViewedSections.Add(NewView);
 	}
 
-	Row.CopyAction(FUIAction(FExecuteAction::CreateSP(this, &FSectionList::OnCopySectionItem, LodIndex, Item.SectionIndex), FCanExecuteAction::CreateSP(this, &FSectionList::OnCanCopySectionItem, LodIndex, Item.SectionIndex)));
-	Row.PasteAction(FUIAction(FExecuteAction::CreateSP(this, &FSectionList::OnPasteSectionItem, LodIndex, Item.SectionIndex)));
-
-	if(SectionListDelegates.OnEnableSectionItem.IsBound())
+	//Chunk section cannot be copy enable or disable, do the operation on the parent section
+	if (!bIsChunkSection)
 	{
-		Row.AddCustomContextMenuAction(FUIAction(FExecuteAction::CreateSP(this, &FSectionList::OnEnableSectionItem, LodIndex, Item.SectionIndex, true)), LOCTEXT("SectionItemContexMenu_Enable", "Enable"));
-		Row.AddCustomContextMenuAction(FUIAction(FExecuteAction::CreateSP(this, &FSectionList::OnEnableSectionItem, LodIndex, Item.SectionIndex, false)), LOCTEXT("SectionItemContexMenu_Disable", "Disable"));
+		Row.CopyAction(FUIAction(FExecuteAction::CreateSP(this, &FSectionList::OnCopySectionItem, LodIndex, Item.SectionIndex), FCanExecuteAction::CreateSP(this, &FSectionList::OnCanCopySectionItem, LodIndex, Item.SectionIndex)));
+		Row.PasteAction(FUIAction(FExecuteAction::CreateSP(this, &FSectionList::OnPasteSectionItem, LodIndex, Item.SectionIndex)));
+
+		if (SectionListDelegates.OnEnableSectionItem.IsBound())
+		{
+			Row.AddCustomContextMenuAction(FUIAction(FExecuteAction::CreateSP(this, &FSectionList::OnEnableSectionItem, LodIndex, Item.SectionIndex, true)), LOCTEXT("SectionItemContexMenu_Enable", "Enable"));
+			Row.AddCustomContextMenuAction(FUIAction(FExecuteAction::CreateSP(this, &FSectionList::OnEnableSectionItem, LodIndex, Item.SectionIndex, false)), LOCTEXT("SectionItemContexMenu_Disable", "Disable"));
+		}
 	}
 
 	Row.NameContent()
diff --git a/Engine/Source/Editor/PropertyEditor/Public/PropertyCustomizationHelpers.h b/Engine/Source/Editor/PropertyEditor/Public/PropertyCustomizationHelpers.h
--- a/Engine/Source/Editor/PropertyEditor/Public/PropertyCustomizationHelpers.h
+++ b/Engine/Source/Editor/PropertyEditor/Public/PropertyCustomizationHelpers.h
@@ -576,7 +576,7 @@ public:
 	* @param Section		The Section being used
 	* @param bCanBeReplced	Whether or not the Section can be replaced by a user
 	*/
-	virtual void AddSection(int32 LodIndex, int32 SectionIndex, FName InMaterialSlotName, int32 InMaterialSlotIndex, FName InOriginalMaterialSlotName, const TMap<int32, FName> &InAvailableMaterialSlotName, const UMaterialInterface* Material, bool IsSectionUsingCloth) = 0;
+	virtual void AddSection(int32 LodIndex, int32 SectionIndex, FName InMaterialSlotName, int32 InMaterialSlotIndex, FName InOriginalMaterialSlotName, const TMap<int32, FName> &InAvailableMaterialSlotName, const UMaterialInterface* Material, bool IsSectionUsingCloth, bool bIsChunkSection) = 0;
 };
 
 
@@ -607,8 +607,10 @@ struct FSectionListItem
 	/* Available material slot name*/
 	TMap<int32, FName> AvailableMaterialSlotName;
 
+	bool bIsChunkSection;
 
-	FSectionListItem(int32 InLodIndex, int32 InSectionIndex, FName InMaterialSlotName, int32 InMaterialSlotIndex, FName InOriginalMaterialSlotName, const TMap<int32, FName> &InAvailableMaterialSlotName, const UMaterialInterface* InMaterial, bool InIsSectionUsingCloth, int32 InThumbnailSize)
+
+	FSectionListItem(int32 InLodIndex, int32 InSectionIndex, FName InMaterialSlotName, int32 InMaterialSlotIndex, FName InOriginalMaterialSlotName, const TMap<int32, FName> &InAvailableMaterialSlotName, const UMaterialInterface* InMaterial, bool InIsSectionUsingCloth, int32 InThumbnailSize, bool InIsChunkSection)
 		: LodIndex(InLodIndex)
 		, SectionIndex(InSectionIndex)
 		, IsSectionUsingCloth(InIsSectionUsingCloth)
@@ -618,6 +620,7 @@ struct FSectionListItem
 		, MaterialSlotIndex(InMaterialSlotIndex)
 		, OriginalMaterialSlotName(InOriginalMaterialSlotName)
 		, AvailableMaterialSlotName(InAvailableMaterialSlotName)
+		, bIsChunkSection(InIsChunkSection)
 	{}
 
 	bool operator==(const FSectionListItem& Other) const
diff --git a/Engine/Source/Editor/SkeletalMeshEditor/Private/SkeletalMeshEditor.cpp b/Engine/Source/Editor/SkeletalMeshEditor/Private/SkeletalMeshEditor.cpp
--- a/Engine/Source/Editor/SkeletalMeshEditor/Private/SkeletalMeshEditor.cpp
+++ b/Engine/Source/Editor/SkeletalMeshEditor/Private/SkeletalMeshEditor.cpp
@@ -44,6 +44,8 @@
 #include "EditorFramework/AssetImportData.h"
 #include "Factories/FbxSkeletalMeshImportData.h"
 
+#include "Misc/MessageDialog.h"
+
 const FName SkeletalMeshEditorAppIdentifier = FName(TEXT("SkeletalMeshEditorApp"));
 
 const FName SkeletalMeshEditorModes::SkeletalMeshEditorMode(TEXT("SkeletalMeshEditorMode"));
@@ -78,6 +80,112 @@ FSkeletalMeshEditor::~FSkeletalMeshEditor()
 	}
 }
 
+bool IsReductionParentBaseLODUseSkeletalMeshBuildWorkflow(USkeletalMesh* SkeletalMesh, int32 TestLODIndex)
+{
+	FSkeletalMeshLODInfo* LODInfo = SkeletalMesh->GetLODInfo(TestLODIndex);
+	if (LODInfo == nullptr || !SkeletalMesh->GetImportedModel() || !SkeletalMesh->GetImportedModel()->LODModels.IsValidIndex(TestLODIndex))
+	{
+		return false;
+	}
+
+	if (SkeletalMesh->GetImportedModel()->LODModels[TestLODIndex].RawSkeletalMeshBulkData.IsBuildDataAvailable())
+	{
+		return true;
+	}
+	if (LODInfo->bHasBeenSimplified || SkeletalMesh->IsReductionActive(TestLODIndex))
+	{
+		int32 ReduceBaseLOD = LODInfo->ReductionSettings.BaseLOD;
+		if (ReduceBaseLOD < TestLODIndex)
+		{
+			return IsReductionParentBaseLODUseSkeletalMeshBuildWorkflow(SkeletalMesh, ReduceBaseLOD);
+		}
+	}
+	return false;
+}
+
+bool FSkeletalMeshEditor::OnRequestClose()
+{
+	bool bAllowClose = true;
+
+	if (PersonaToolkit.IsValid() && SkeletalMesh)
+	{
+		bool bHaveModifiedLOD = false;
+		for (int32 LODIndex = 0; LODIndex < SkeletalMesh->GetLODNum(); ++LODIndex)
+		{
+			FSkeletalMeshLODInfo* LODInfo = SkeletalMesh->GetLODInfo(LODIndex);
+			if (LODInfo == nullptr || !SkeletalMesh->GetImportedModel() || !SkeletalMesh->GetImportedModel()->LODModels.IsValidIndex(LODIndex))
+			{
+				continue;
+			}
+			
+			//Do not prevent exiting if we are not using the skeletal mesh build workflow
+			if (!SkeletalMesh->GetImportedModel()->LODModels[LODIndex].RawSkeletalMeshBulkData.IsBuildDataAvailable())
+			{
+				if (!LODInfo->bHasBeenSimplified && !SkeletalMesh->IsReductionActive(LODIndex))
+				{
+					continue;
+				}
+				//Do not prevent exit if the generated LOD is base on a LODModel not using the skeletalmesh build workflow
+				int32 ReduceBaseLOD = LODInfo->ReductionSettings.BaseLOD;
+				if(!IsReductionParentBaseLODUseSkeletalMeshBuildWorkflow(SkeletalMesh, ReduceBaseLOD))
+				{
+					continue;
+				}
+				
+			}
+
+			bool bValidLODSettings = false;
+			if (SkeletalMesh->LODSettings != nullptr)
+			{
+				const int32 NumSettings = FMath::Min(SkeletalMesh->LODSettings->GetNumberOfSettings(), SkeletalMesh->GetLODNum());
+				if (LODIndex < NumSettings)
+				{
+					bValidLODSettings = true;
+				}
+			}
+			const FSkeletalMeshLODGroupSettings* SkeletalMeshLODGroupSettings = bValidLODSettings ? &SkeletalMesh->LODSettings->GetSettingsForLODLevel(LODIndex) : nullptr;
+
+			FGuid BuildGUID = LODInfo->ComputeDeriveDataCacheKey(SkeletalMeshLODGroupSettings);
+			if (LODInfo->BuildGUID != BuildGUID)
+			{
+				bHaveModifiedLOD = true;
+				break;
+			}
+			FString BuildStringID = SkeletalMesh->GetImportedModel()->LODModels[LODIndex].GetLODModelDeriveDataKey();
+			if (SkeletalMesh->GetImportedModel()->LODModels[LODIndex].BuildStringID != BuildStringID)
+			{
+				bHaveModifiedLOD = true;
+				break;
+			}
+		}
+
+		if (bHaveModifiedLOD)
+		{
+			// find out the user wants to do with this dirty material
+			EAppReturnType::Type OkCancelReply = FMessageDialog::Open(
+				EAppMsgType::OkCancel,
+				FText::Format(LOCTEXT("SkeletalMeshEditorShouldApplyLODChanges", "We have to apply level of detail changes to {0} before exiting the skeletal mesh editor."), FText::FromString(PersonaToolkit->GetMesh()->GetName()))
+			);
+
+			switch (OkCancelReply)
+			{
+			case EAppReturnType::Ok:
+				{
+					SkeletalMesh->MarkPackageDirty();
+					SkeletalMesh->PostEditChange();
+					bAllowClose = true;
+				}
+				break;
+			case EAppReturnType::Cancel:
+				// Don't exit.
+				bAllowClose = false;
+				break;
+			}
+		}
+	}
+	return bAllowClose;
+}
+
 void FSkeletalMeshEditor::RegisterTabSpawners(const TSharedRef<class FTabManager>& InTabManager)
 {
 	WorkspaceMenuCategory = InTabManager->AddLocalWorkspaceMenuCategory(LOCTEXT("WorkspaceMenu_SkeletalMeshEditor", "Skeletal Mesh Editor"));
@@ -546,8 +654,7 @@ void FSkeletalMeshEditor::OnRemoveSectionFromLodAndBelowMenuItemClicked(int32 Lo
 		const FSkeletalMeshLODInfo* CurrentSkeletalMeshLODInfo = SkeletalMesh->GetLODInfo(GenerateLodIndex);
 		if (CurrentSkeletalMeshLODInfo != nullptr && CurrentSkeletalMeshLODInfo->bHasBeenSimplified && BaseLodIndexes.Contains(CurrentSkeletalMeshLODInfo->ReductionSettings.BaseLOD))
 		{
-			bool bRestoreClothing = true;
-			FLODUtilities::SimplifySkeletalMeshLOD(UpdateContext, GenerateLodIndex, true, bRestoreClothing);
+			FLODUtilities::SimplifySkeletalMeshLOD(UpdateContext, GenerateLodIndex, true);
 			BaseLodIndexes.Add(GenerateLodIndex);
 		}
 	}
@@ -668,7 +775,7 @@ void FSkeletalMeshEditor::OnCreateClothingAssetMenuItemClicked(FSkeletalMeshClot
 						{
 							SectionIndex = i;
 							AssetLodIndex = SkelLod.Sections[i].ClothingData.AssetLodIndex;
-							TargetAssetPtr->UnbindFromSkeletalMesh(Mesh, Params.TargetLod);
+							RemoveClothing(Params.TargetLod, SectionIndex);
 							break;
 						}
 					}
@@ -786,13 +893,53 @@ void FSkeletalMeshEditor::ApplyClothing(UClothingAssetBase* InAsset, int32 InLod
 {
 	USkeletalMesh* Mesh = GetPersonaToolkit()->GetPreviewMesh();
 
-	if(UClothingAsset* ClothingAsset = Cast<UClothingAsset>(InAsset))
+	if (Mesh->GetImportedModel() == nullptr || !Mesh->GetImportedModel()->LODModels.IsValidIndex(InLodIndex))
 	{
-		ClothingAsset->BindToSkeletalMesh(Mesh, InLodIndex, InSectionIndex, InClothingLod);
+		return;
 	}
-	else
+
+	FSkeletalMeshLODModel& LODModel = Mesh->GetImportedModel()->LODModels[InLodIndex];
+	const FSkelMeshSection& Section = LODModel.Sections[InSectionIndex];
 	{
-		RemoveClothing(InLodIndex, InSectionIndex);
+		FScopedSkeletalMeshPostEditChange ScopedPostEditChange(Mesh);
+		FScopedTransaction Transaction(LOCTEXT("SkeletalMeshEditorApplyClothingTransaction", "Persona editor: Apply Section Cloth"));
+		Mesh->Modify();
+
+		FSkelMeshSourceSectionUserData& OriginalSectionData = LODModel.UserSectionsData.FindOrAdd(Section.OriginalDataSectionIndex);
+		auto ClearOriginalSectionUserData = [&OriginalSectionData]()
+		{
+			OriginalSectionData.CorrespondClothAssetIndex = INDEX_NONE;
+			OriginalSectionData.ClothingData.AssetGuid = FGuid();
+			OriginalSectionData.ClothingData.AssetLodIndex = INDEX_NONE;
+		};
+		if (UClothingAsset* ClothingAsset = Cast<UClothingAsset>(InAsset))
+		{
+			// Look for a currently bound asset an unbind it if necessary first
+			if (UClothingAssetBase* CurrentAsset = Mesh->GetSectionClothingAsset(InLodIndex, InSectionIndex))
+			{
+				CurrentAsset->UnbindFromSkeletalMesh(Mesh, InLodIndex);
+				ClearOriginalSectionUserData();
+			}
+
+			if (ClothingAsset->BindToSkeletalMesh(Mesh, InLodIndex, InSectionIndex, InClothingLod))
+			{
+				//Successful bind so set the SectionUserData
+				int32 AssetIndex = INDEX_NONE;
+				check(Mesh->MeshClothingAssets.Find(ClothingAsset, AssetIndex));
+				OriginalSectionData.CorrespondClothAssetIndex = AssetIndex;
+				OriginalSectionData.ClothingData.AssetGuid = ClothingAsset->GetAssetGuid();
+				OriginalSectionData.ClothingData.AssetLodIndex = InClothingLod;
+			}
+		}
+		else if (Mesh)
+		{
+			//User set none, so unbind anything that is bind
+			if (UClothingAssetBase* CurrentAsset = Mesh->GetSectionClothingAsset(InLodIndex, InSectionIndex))
+			{
+				CurrentAsset->UnbindFromSkeletalMesh(Mesh, InLodIndex);
+				ClearOriginalSectionUserData();
+			}
+		}
 	}
 }
 
@@ -800,11 +947,30 @@ void FSkeletalMeshEditor::RemoveClothing(int32 InLodIndex, int32 InSectionIndex)
 {
 	USkeletalMesh* Mesh = GetPersonaToolkit()->GetPreviewMesh();
 
-	if(Mesh)
+	if (Mesh->GetImportedModel() == nullptr || !Mesh->GetImportedModel()->LODModels.IsValidIndex(InLodIndex))
+	{
+		return;
+	}
+
+	FSkeletalMeshLODModel& LODModel = Mesh->GetImportedModel()->LODModels[InLodIndex];
+	const FSkelMeshSection& Section = LODModel.Sections[InSectionIndex];
 	{
-		if(UClothingAssetBase* CurrentAsset = Mesh->GetSectionClothingAsset(InLodIndex, InSectionIndex))
+		FScopedSkeletalMeshPostEditChange ScopedPostEditChange(Mesh);
+		FScopedTransaction Transaction(LOCTEXT("SkeletalMeshEditorRemoveClothingTransaction", "Persona editor: Remove Section Cloth"));
+		Mesh->Modify();
+
+		FSkelMeshSourceSectionUserData& OriginalSectionData = LODModel.UserSectionsData.FindOrAdd(Section.OriginalDataSectionIndex);
+		auto ClearOriginalSectionUserData = [&OriginalSectionData]()
+		{
+			OriginalSectionData.CorrespondClothAssetIndex = INDEX_NONE;
+			OriginalSectionData.ClothingData.AssetGuid = FGuid();
+			OriginalSectionData.ClothingData.AssetLodIndex = INDEX_NONE;
+		};
+		// Look for a currently bound asset an unbind it if necessary first
+		if (UClothingAssetBase* CurrentAsset = Mesh->GetSectionClothingAsset(InLodIndex, InSectionIndex))
 		{
 			CurrentAsset->UnbindFromSkeletalMesh(Mesh, InLodIndex);
+			ClearOriginalSectionUserData();
 		}
 	}
 }
@@ -881,15 +1047,65 @@ UObject* FSkeletalMeshEditor::HandleGetAsset()
 	return GetEditingObject();
 }
 
-bool FSkeletalMeshEditor::HandleReimportMeshInternal(int32 SourceFileIndex /*= INDEX_NONE*/, bool bWithNewFile /*= false*/)
+class FScopedSuspendAlternateSkinWeightPreview
 {
-	// Make sure we clear any skin weight profile currently being previewed
-	UDebugSkelMeshComponent* MeshComponent = GetPersonaToolkit()->GetPreviewScene()->GetPreviewMeshComponent();
-	const FName ProfileName = MeshComponent ? MeshComponent->GetCurrentSkinWeightProfileName() : NAME_None;
-	if (MeshComponent)
+public:
+	/*
+	 * This constructor suspend the alternate skinning preview for all editor component that use the specified skeletalmesh
+	 * Parameters:
+	 * @param InSkeletalMesh - SkeletalMesh use to know which preview component we have to suspend the alternate skinning preview.
+	 */
+	FScopedSuspendAlternateSkinWeightPreview(USkeletalMesh* InSkeletalMesh);
+
+	/*
+	 * This destructor put back the preview alternate skinning
+	 */
+	~FScopedSuspendAlternateSkinWeightPreview();
+
+private:
+	TArray< TTuple<UDebugSkelMeshComponent*, FName> > SuspendedComponentArray;
+};
+
+FScopedSuspendAlternateSkinWeightPreview::FScopedSuspendAlternateSkinWeightPreview(USkeletalMesh* SkeletalMesh)
+{
+	SuspendedComponentArray.Empty(2);
+	if (SkeletalMesh != nullptr)
 	{
-		MeshComponent->ClearSkinWeightProfile();
+		// Now iterate over all skeletal mesh components and unregister them from the world, we will reregister them in the destructor
+		for (TObjectIterator<UDebugSkelMeshComponent> It; It; ++It)
+		{
+			UDebugSkelMeshComponent* DebugSKComp = *It;
+			if (DebugSKComp->SkeletalMesh == SkeletalMesh)
+			{
+				const FName ProfileName = DebugSKComp->GetCurrentSkinWeightProfileName();
+				if (ProfileName != NAME_None)
+				{
+					DebugSKComp->ClearSkinWeightProfile();
+					TTuple<UDebugSkelMeshComponent*, FName> ComponentTupple;
+					ComponentTupple.Key = DebugSKComp;
+					ComponentTupple.Value = ProfileName;
+					SuspendedComponentArray.Add(ComponentTupple);
+				}
+			}
+		}
 	}
+}
+
+FScopedSuspendAlternateSkinWeightPreview::~FScopedSuspendAlternateSkinWeightPreview()
+{
+	//Put back the skin weight profile for all editor debug component
+	for (const TTuple<UDebugSkelMeshComponent*, FName>& ComponentTupple : SuspendedComponentArray)
+	{
+		ComponentTupple.Key->SetSkinWeightProfile(ComponentTupple.Value);
+	}
+	SuspendedComponentArray.Empty();
+}
+
+
+bool FSkeletalMeshEditor::HandleReimportMeshInternal(int32 SourceFileIndex /*= INDEX_NONE*/, bool bWithNewFile /*= false*/)
+{
+	FScopedSkeletalMeshPostEditChange ScopedPostEditChange(SkeletalMesh);
+	FScopedSuspendAlternateSkinWeightPreview ScopedSuspendAlternateSkinnWeightPreview(SkeletalMesh);
 
 	// Reimport the asset
 	const bool bResult = FReimportManager::Instance()->Reimport(SkeletalMesh, true, true, TEXT(""), nullptr, SourceFileIndex, bWithNewFile);
@@ -897,27 +1113,28 @@ bool FSkeletalMeshEditor::HandleReimportMeshInternal(int32 SourceFileIndex /*= I
 	// Refresh skeleton tree
 	SkeletonTree->Refresh();
 
-	// Re-set skin weigh profile
-	if (MeshComponent && (ProfileName != NAME_None))
-	{
-		MeshComponent->SetSkinWeightProfile(ProfileName);
-	}
-
 	return bResult;
 }
 
 void FSkeletalMeshEditor::HandleReimportMesh(int32 SourceFileIndex /*= INDEX_NONE*/)
 {
+	FScopedSuspendAlternateSkinWeightPreview ScopedSuspendAlternateSkinnWeightPreview(SkeletalMesh);
+
 	HandleReimportMeshInternal(SourceFileIndex, false);
 }
 
 void FSkeletalMeshEditor::HandleReimportMeshWithNewFile(int32 SourceFileIndex /*= INDEX_NONE*/)
 {
+	FScopedSuspendAlternateSkinWeightPreview ScopedSuspendAlternateSkinnWeightPreview(SkeletalMesh);
+
 	HandleReimportMeshInternal(SourceFileIndex, true);
 }
 
 void ReimportAllCustomLODs(USkeletalMesh* SkeletalMesh, UDebugSkelMeshComponent* PreviewMeshComponent, bool bWithNewFile)
 {
+	FScopedSkeletalMeshPostEditChange ScopedPostEditChange(SkeletalMesh);
+	FScopedSuspendAlternateSkinWeightPreview ScopedSuspendAlternateSkinnWeightPreview(SkeletalMesh);
+
 	//Find the dependencies of the generated LOD
 	TArray<bool> Dependencies;
 	Dependencies.AddZeroed(SkeletalMesh->GetLODNum());
@@ -939,7 +1156,7 @@ void ReimportAllCustomLODs(USkeletalMesh* SkeletalMesh, UDebugSkelMeshComponent*
 				SkeletalMesh->GetLODInfo(LodIndex)->SourceImportFilename.Empty();
 			}
 
-			if (!FbxMeshUtils::ImportMeshLODDialog(SkeletalMesh, LodIndex))
+			if (!FbxMeshUtils::ImportMeshLODDialog(SkeletalMesh, LodIndex, false))
 			{
 				if (bWithNewFile)
 				{
@@ -957,8 +1174,7 @@ void ReimportAllCustomLODs(USkeletalMesh* SkeletalMesh, UDebugSkelMeshComponent*
 			FSkeletalMeshUpdateContext UpdateContext;
 			UpdateContext.SkeletalMesh = SkeletalMesh;
 			UpdateContext.AssociatedComponents.Push(PreviewMeshComponent);
-			bool bRestoreClothing = true;
-			FLODUtilities::SimplifySkeletalMeshLOD(UpdateContext, LodIndex, true, bRestoreClothing);
+			FLODUtilities::SimplifySkeletalMeshLOD(UpdateContext, LodIndex, false);
 			Dependencies[LodIndex] = true;
 		}
 	}
@@ -966,9 +1182,12 @@ void ReimportAllCustomLODs(USkeletalMesh* SkeletalMesh, UDebugSkelMeshComponent*
 
 void FSkeletalMeshEditor::HandleReimportAllMesh(int32 SourceFileIndex /*= INDEX_NONE*/)
 {
+	FScopedSuspendAlternateSkinWeightPreview ScopedSuspendAlternateSkinnWeightPreview(SkeletalMesh);
 	// Reimport the asset
 	if (SkeletalMesh)
 	{
+		FScopedSkeletalMeshPostEditChange ScopedPostEditChange(SkeletalMesh);
+
 		//Reimport base LOD
 		if (HandleReimportMeshInternal(SourceFileIndex, false))
 		{
@@ -980,9 +1199,12 @@ void FSkeletalMeshEditor::HandleReimportAllMesh(int32 SourceFileIndex /*= INDEX_
 
 void FSkeletalMeshEditor::HandleReimportAllMeshWithNewFile(int32 SourceFileIndex /*= INDEX_NONE*/)
 {
+	FScopedSuspendAlternateSkinWeightPreview ScopedSuspendAlternateSkinnWeightPreview(SkeletalMesh);
+
 	// Reimport the asset
 	if (SkeletalMesh)
 	{
+		FScopedSkeletalMeshPostEditChange ScopedPostEditChange(SkeletalMesh);
 		TArray<UObject*> ImportObjs;
 		ImportObjs.Add(SkeletalMesh);
 		if (HandleReimportMeshInternal(SourceFileIndex, true))
diff --git a/Engine/Source/Editor/StaticMeshEditor/Private/StaticMeshEditorTools.cpp b/Engine/Source/Editor/StaticMeshEditor/Private/StaticMeshEditorTools.cpp
--- a/Engine/Source/Editor/StaticMeshEditor/Private/StaticMeshEditorTools.cpp
+++ b/Engine/Source/Editor/StaticMeshEditor/Private/StaticMeshEditorTools.cpp
@@ -1971,7 +1971,7 @@ void FMeshSectionSettingsLayout::OnGetSectionsForView(ISectionListBuilder& OutSe
 				{
 					SectionMaterial = UMaterial::GetDefaultMaterial(MD_Surface);
 				}
-				OutSections.AddSection(LODIndex, SectionIndex, CurrentSectionMaterialSlotName, MaterialIndex, CurrentSectionOriginalImportedMaterialName, AvailableSectionName, StaticMesh.StaticMaterials[MaterialIndex].MaterialInterface, false);
+				OutSections.AddSection(LODIndex, SectionIndex, CurrentSectionMaterialSlotName, MaterialIndex, CurrentSectionOriginalImportedMaterialName, AvailableSectionName, StaticMesh.StaticMaterials[MaterialIndex].MaterialInterface, false, false);
 			}
 		}
 	}
