Revert the meshutilities build cs
#rnx

diff --git a/Engine/Source/Developer/MeshUtilities/MeshUtilities.Build.cs b/Engine/Source/Developer/MeshUtilities/MeshUtilities.Build.cs
--- a/Engine/Source/Developer/MeshUtilities/MeshUtilities.Build.cs
+++ b/Engine/Source/Developer/MeshUtilities/MeshUtilities.Build.cs
@@ -33,6 +33,7 @@ public class MeshUtilities : ModuleRules
 				"PropertyEditor",
 				"EditorStyle",
                 "GraphColor",
+                "MeshBuilder",
                 "MeshUtilitiesCommon",
                 "MeshDescription",
 				"StaticMeshDescription",
