Make sure we do not test against zero but against nearly zero when computing the tangents for skeletalmesh.

diff --git a/Engine/Source/Developer/MeshUtilities/Private/MeshUtilities.cpp b/Engine/Source/Developer/MeshUtilities/Private/MeshUtilities.cpp
--- a/Engine/Source/Developer/MeshUtilities/Private/MeshUtilities.cpp
+++ b/Engine/Source/Developer/MeshUtilities/Private/MeshUtilities.cpp
@@ -3588,8 +3588,8 @@ public:
 			bool bSkipNTB = true;
 			for (int32 CornerIndex = 0; CornerIndex < 3; CornerIndex++)
 			{
-				bCornerHasNormal[CornerIndex] = !WedgeTangentZ[WedgeOffset + CornerIndex].IsZero();
-				bCornerHasTangents[CornerIndex] = !WedgeTangentX[WedgeOffset + CornerIndex].IsZero() && !WedgeTangentY[WedgeOffset + CornerIndex].IsZero();
+				bCornerHasNormal[CornerIndex] = !WedgeTangentZ[WedgeOffset + CornerIndex].IsNearlyZero();
+				bCornerHasTangents[CornerIndex] = !WedgeTangentX[WedgeOffset + CornerIndex].IsNearlyZero() && !WedgeTangentY[WedgeOffset + CornerIndex].IsNearlyZero();
 
 				//If we want to compute mikkt we dont check tangents to skip this corner
 				if (!bCornerHasNormal[CornerIndex] || (!bUseMikktSpace && !bCornerHasTangents[CornerIndex]))
@@ -3814,15 +3814,15 @@ public:
 							}
 						}
 					}
-					if (!WedgeTangentX[WedgeOffset + CornerIndex].IsZero())
+					if (!WedgeTangentX[WedgeOffset + CornerIndex].IsNearlyZero())
 					{
 						CornerTangentX[CornerIndex] = WedgeTangentX[WedgeOffset + CornerIndex];
 					}
-					if (!WedgeTangentY[WedgeOffset + CornerIndex].IsZero())
+					if (!WedgeTangentY[WedgeOffset + CornerIndex].IsNearlyZero())
 					{
 						CornerTangentY[CornerIndex] = WedgeTangentY[WedgeOffset + CornerIndex];
 					}
-					if (!WedgeTangentZ[WedgeOffset + CornerIndex].IsZero())
+					if (!WedgeTangentZ[WedgeOffset + CornerIndex].IsNearlyZero())
 					{
 						CornerNormal[CornerIndex] = WedgeTangentZ[WedgeOffset + CornerIndex];
 					}
