Use the polygon shell (connected edges+uvs+normals+color) when chunking sections, to avoid adding vertices when chunking.

diff --git a/Engine/Source/Developer/MeshUtilities/Private/SkeletalMeshTools.cpp b/Engine/Source/Developer/MeshUtilities/Private/SkeletalMeshTools.cpp
--- a/Engine/Source/Developer/MeshUtilities/Private/SkeletalMeshTools.cpp
+++ b/Engine/Source/Developer/MeshUtilities/Private/SkeletalMeshTools.cpp
@@ -184,6 +184,236 @@ namespace SkeletalMeshTools
 		}
 	}
 
+
+	namespace PolygonShellsHelper
+	{
+		struct FPatchAndBoneInfluence
+		{
+			TArray<FBoneIndexType> UniqueBones;
+			TArray<int32> PatchToChunkWith;
+			bool bIsParent = false;
+		};
+		//This function add every triangles connected to the triangle queue.
+		//A connected triangle pair must share at least 1 vertex between the two triangles.
+		//If bConnectByEdge is true, the connected triangle must share at least one edge (two vertex index)
+		//To have a connected vertex instance pair, the position, NTBs, UVs(channel 0) and color must match.
+		void AddAdjacentFace(const TArray<uint32>& Indices, const TArray<FSoftSkinBuildVertex>& Vertices, TBitArray<>& FaceAdded, const TMap<int32, TArray<int32>>& VertexIndexToAdjacentFaces, const int32 FaceIndex, TArray<int32>& TriangleQueue, const bool bConnectByEdge)
+		{
+			int32 NumFaces = Indices.Num()/3;
+			check(FaceAdded.Num() == NumFaces);
+
+			TMap<int32, int32> AdjacentFaceCommonVertices;
+			for (int32 Corner = 0; Corner < 3; Corner++)
+			{
+				int32 IndiceIndex = FaceIndex * 3 + Corner;
+				checkSlow(Indices.IsValidIndex(IndiceIndex));
+				int32 VertexIndex = Indices[IndiceIndex];
+				checkSlow(Vertices.IsValidIndex(VertexIndex));
+				const FSoftSkinBuildVertex& SoftSkinVertRef = Vertices[VertexIndex];
+				const FVector& PositionRef = SoftSkinVertRef.Position;
+				const FVector& TangentXRef = SoftSkinVertRef.TangentX;
+				const FVector& TangentYRef = SoftSkinVertRef.TangentY;
+				const FVector& TangentZRef = SoftSkinVertRef.TangentZ;
+				const FVector2D& UVRef = SoftSkinVertRef.UVs[0];
+				const FColor& ColorRef = SoftSkinVertRef.Color;
+				const TArray<int32>& AdjacentFaces = VertexIndexToAdjacentFaces.FindChecked(VertexIndex);
+				for (int32 AdjacentFaceArrayIndex = 0; AdjacentFaceArrayIndex < AdjacentFaces.Num(); ++AdjacentFaceArrayIndex)
+				{
+					const int32 AdjacentFaceIndex = AdjacentFaces[AdjacentFaceArrayIndex];
+					if (!FaceAdded[AdjacentFaceIndex] && AdjacentFaceIndex != FaceIndex)
+					{
+						//Ensure we have position, NTBs, uv and color match to allow a connection.
+						bool bRealConnection = false;
+						for (int32 AdjCorner = 0; AdjCorner < 3; AdjCorner++)
+						{
+							const int32 IndiceIndexAdj = AdjacentFaceIndex * 3 + AdjCorner;
+							checkSlow(Indices.IsValidIndex(IndiceIndexAdj));
+							const int32 VertexIndexAdj = Indices[IndiceIndexAdj];
+							checkSlow(Vertices.IsValidIndex(VertexIndexAdj));
+							const FSoftSkinBuildVertex& SoftSkinVertAdj = Vertices[VertexIndexAdj];
+							if (PositionRef.Equals(SoftSkinVertAdj.Position, SMALL_NUMBER) &&
+								TangentXRef.Equals(SoftSkinVertAdj.TangentX, SMALL_NUMBER) &&
+								TangentYRef.Equals(SoftSkinVertAdj.TangentY, SMALL_NUMBER) &&
+								TangentZRef.Equals(SoftSkinVertAdj.TangentZ, SMALL_NUMBER) &&
+								UVRef.Equals(SoftSkinVertAdj.UVs[0], KINDA_SMALL_NUMBER) &&
+								ColorRef == SoftSkinVertAdj.Color)
+							{
+								bRealConnection = true;
+								break;
+							}
+						}
+						if (!bRealConnection)
+						{
+							continue;
+						}
+
+						bool bAddConnected = !bConnectByEdge;
+						if (bConnectByEdge)
+						{
+							int32& AdjacentFaceCommonVerticeCount = AdjacentFaceCommonVertices.FindOrAdd(AdjacentFaceIndex);
+							AdjacentFaceCommonVerticeCount++;
+							//Is the connected triangles share 2 vertex index (one edge) not only one vertex
+							bAddConnected = AdjacentFaceCommonVerticeCount > 1;
+						}
+
+						if (bAddConnected)
+						{
+							TriangleQueue.Add(AdjacentFaceIndex);
+							//Add the face only once by marking the face has computed
+							FaceAdded[AdjacentFaceIndex] = true;
+						}
+					}
+				}
+			}
+		}
+
+
+		//Fill FaceIndexToPatchIndex so every triangle knows its unique island patch index.
+		//Each island patch have is fill with connected vertexinstance where position, NTBs. UVs and colors are nearly equal.
+		//@Param bConnectByEdge: If true we need at least 2 vertex index (one edge) to connect 2 triangles. If false we just need one vertex index (bowtie)
+		void FillPolygonPatch(const TArray<uint32>& Indices, const TArray<FSoftSkinBuildVertex>& Vertices, const TMap<uint32, TArray<FBoneIndexType>>& AlternateBoneIDs, TArray<FPatchAndBoneInfluence>& PatchData, TArray<int32>& IndiceToPatchIndex, const int32 MaxBonesPerChunk, const bool bConnectByEdge)
+		{
+			const int32 NumIndice = Indices.Num();
+			check(IndiceToPatchIndex.Num() == NumIndice);
+			const int32 NumFace = NumIndice / 3;
+			int32 PatchIndex = 0;
+
+			//Store a map containing connected faces for each vertex index
+			TMap<int32, TArray<int32>> VertexIndexToAdjacentFaces;
+			VertexIndexToAdjacentFaces.Reserve(Vertices.Num());
+			for (int32 FaceIndex = 0; FaceIndex < NumFace; ++FaceIndex)
+			{
+				const int32 IndiceOffset = FaceIndex * 3;
+				for (int32 Corner = 0; Corner < 3; Corner++)
+				{
+					checkSlow(Indices.IsValidIndex(IndiceOffset + Corner));
+					int32 VertexIndex = Indices[IndiceOffset + Corner];
+					TArray<int32>& AdjacentFaces = VertexIndexToAdjacentFaces.FindOrAdd(VertexIndex);
+					AdjacentFaces.AddUnique(FaceIndex);
+				}
+			}
+
+			//Mark added face so we do not add them more then once
+			TBitArray<> FaceAdded;
+			FaceAdded.Init(false, NumFace);
+
+			TArray<int32> TriangleQueue;
+			TriangleQueue.Reserve(100);
+			for (int32 FaceIndex = 0; FaceIndex < NumFace; ++FaceIndex)
+			{
+				//Skip already added faces
+				if (FaceAdded[FaceIndex])
+				{
+					continue;
+				}
+
+				//Add all the faces connected to the current face index
+				TriangleQueue.Reset();
+				TriangleQueue.Add(FaceIndex); //Use a queue to avoid recursive function
+				FaceAdded[FaceIndex] = true;
+				while (TriangleQueue.Num() > 0)
+				{
+					int32 CurrentTriangleIndex = TriangleQueue.Pop(false);
+					int32 IndiceOffset = CurrentTriangleIndex * 3;
+					for (int32 Corner = 0; Corner < 3; Corner++)
+					{
+						const int32 IndiceIndex = IndiceOffset + Corner;
+						checkSlow(IndiceToPatchIndex.IsValidIndex(IndiceIndex));
+						IndiceToPatchIndex[IndiceIndex] = PatchIndex;
+						checkSlow(Indices.IsValidIndex(IndiceIndex));
+						const int32 VertexIndex = Indices[IndiceIndex];
+						checkSlow(Vertices.IsValidIndex(VertexIndex));
+						for(int32 InfluenceIndex = 0; InfluenceIndex < MAX_TOTAL_INFLUENCES; ++InfluenceIndex)
+						{
+							if (Vertices[VertexIndex].InfluenceWeights[InfluenceIndex] == 0)
+							{
+								break;
+							}
+							if (!PatchData.IsValidIndex(PatchIndex))
+							{
+								PatchData.AddDefaulted(PatchData.Num() - PatchIndex + 1);
+							}
+							PatchData[PatchIndex].UniqueBones.AddUnique(Vertices[VertexIndex].InfluenceBones[InfluenceIndex]);
+						}
+						//Add the alternate bones
+						const TArray<FBoneIndexType>* AlternateBones = AlternateBoneIDs.Find(Vertices[VertexIndex].PointWedgeIdx);
+						if (AlternateBones)
+						{
+							for (int32 InfluenceIndex = 0; InfluenceIndex < AlternateBones->Num(); InfluenceIndex++)
+							{
+								if (!PatchData.IsValidIndex(PatchIndex))
+								{
+									PatchData.AddDefaulted(PatchData.Num() - PatchIndex + 1);
+								}
+								PatchData[PatchIndex].UniqueBones.AddUnique((*AlternateBones)[InfluenceIndex]);
+							}
+						}
+					}
+					//The patch should exist at this time
+					checkSlow(PatchData.IsValidIndex(PatchIndex));
+
+					AddAdjacentFace(Indices, Vertices, FaceAdded, VertexIndexToAdjacentFaces, CurrentTriangleIndex, TriangleQueue, bConnectByEdge);
+				}
+				PatchIndex++;
+			}
+		}
+
+		void RecursiveFillRemapIndices(const TArray<FPatchAndBoneInfluence>& PatchData, const int32 PatchIndex, const TArray<int32>& IndiceToPatchIndex, TArray<uint32>& SrcChunkRemapIndicesIndex)
+		{
+			for (int32 IndiceIndex = 0; IndiceIndex < IndiceToPatchIndex.Num(); ++IndiceIndex)
+			{
+				int32 IndicePatchIndex = IndiceToPatchIndex[IndiceIndex];
+				if (PatchIndex == IndicePatchIndex)
+				{
+					SrcChunkRemapIndicesIndex.Add(IndiceIndex);
+				}
+			}
+			checkSlow(PatchData.IsValidIndex(PatchIndex));
+			//Do the child patch to chunk with
+			for (int32 SubPatchIndex = 0; SubPatchIndex < PatchData[PatchIndex].PatchToChunkWith.Num(); ++SubPatchIndex)
+			{
+				RecursiveFillRemapIndices(PatchData, PatchData[PatchIndex].PatchToChunkWith[SubPatchIndex], IndiceToPatchIndex, SrcChunkRemapIndicesIndex);
+			}
+		}
+
+		//Sort the shells to a setup that use the less section possible
+		void GatherShellUsingSameBones(const int32 ParentPatchIndex, TArray<FPatchAndBoneInfluence>& PatchData, TBitArray<>& PatchConsumed, const int32 MaxBonesPerChunk)
+		{
+			checkSlow(PatchData.IsValidIndex(ParentPatchIndex));
+			TArray<FBoneIndexType> UniqueBones = PatchData[ParentPatchIndex].UniqueBones;
+			PatchData[ParentPatchIndex].bIsParent = true;
+			if (UniqueBones.Num() > MaxBonesPerChunk)
+			{
+				return;
+			}
+			for (int32 PatchIndex = ParentPatchIndex + 1; PatchIndex < PatchData.Num(); ++PatchIndex)
+			{
+				if (PatchConsumed[PatchIndex])
+				{
+					continue;
+				}
+
+				TArray<FBoneIndexType> AddedBones;
+				for (int32 BoneIndex = 0; BoneIndex < PatchData[PatchIndex].UniqueBones.Num(); ++BoneIndex)
+				{
+					FBoneIndexType BoneIndexType = PatchData[PatchIndex].UniqueBones[BoneIndex];
+					if (!UniqueBones.Contains(BoneIndexType))
+					{
+						AddedBones.AddUnique(BoneIndexType);
+					}
+				}
+				if (AddedBones.Num() + UniqueBones.Num() <= MaxBonesPerChunk)
+				{
+					UniqueBones.Append(AddedBones);
+					PatchConsumed[PatchIndex] = true;
+					//We only support one parent layer, the assumption is we have a hierarchy depth of max 2 (parents, childs)
+					checkSlow(!PatchData[PatchIndex].bIsParent);
+					PatchData[ParentPatchIndex].PatchToChunkWith.Add(PatchIndex);
+				}
+			}
+		}
+	}
+
 	void ChunkSkinnedVertices(TArray<FSkinnedMeshChunk*>& Chunks, TMap<uint32, TArray<FBoneIndexType>>& AlternateBoneIDs, int32 MaxBonesPerChunk)
 	{
 #if WITH_EDITORONLY_DATA
@@ -201,6 +431,40 @@ namespace SkeletalMeshTools
 		};
 		SrcChunks.Sort(FCompareSkinnedMeshChunk());
 
+		TMap<int32, TArray<PolygonShellsHelper::FPatchAndBoneInfluence>> PatchDataPerSrcChunk;
+		TMap<int32, TArray<int32>> IndiceToPatchIndexPerSrcChunk;
+		
+		//Find the shells inside chunks
+		for (int32 ChunkIndex = 0; ChunkIndex < SrcChunks.Num(); ++ChunkIndex)
+		{
+			FSkinnedMeshChunk* ChunkToShell = SrcChunks[ChunkIndex];
+			TArray<uint32>& Indices = ChunkToShell->Indices;
+			TArray<FSoftSkinBuildVertex>& Vertices = ChunkToShell->Vertices;
+			TArray<PolygonShellsHelper::FPatchAndBoneInfluence>& PatchData = PatchDataPerSrcChunk.Add(ChunkIndex);
+			TArray<int32>& IndiceToPatchIndex = IndiceToPatchIndexPerSrcChunk.Add(ChunkIndex);
+			IndiceToPatchIndex.AddZeroed(Indices.Num());
+			//We need edge connection (2 similar vertex )
+			const bool bConnectByEdge = true;
+			PolygonShellsHelper::FillPolygonPatch(Indices, Vertices, AlternateBoneIDs, PatchData, IndiceToPatchIndex, MaxBonesPerChunk, bConnectByEdge);
+		}
+
+		for (int32 SrcChunkIndex = 0; SrcChunkIndex < SrcChunks.Num(); ++SrcChunkIndex)
+		{
+			TArray<PolygonShellsHelper::FPatchAndBoneInfluence>& PatchData = PatchDataPerSrcChunk[SrcChunkIndex];
+			TBitArray<> PatchConsumed;
+			PatchConsumed.Init(false, PatchData.Num());
+
+			for (int32 PatchIndex = 0; PatchIndex < PatchData.Num(); ++PatchIndex)
+			{
+				if (PatchConsumed[PatchIndex])
+				{
+					continue;
+				}
+				PatchConsumed[PatchIndex] = true;
+				PolygonShellsHelper::GatherShellUsingSameBones(PatchIndex, PatchData, PatchConsumed, MaxBonesPerChunk);
+			}
+		}
+
 		// Now split chunks to respect the desired bone limit.
 		TIndirectArray<TArray<int32> > IndexMaps;
 		TArray<FBoneIndexType, TInlineAllocator<MAX_TOTAL_INFLUENCES*3> > UniqueBones;
@@ -209,95 +473,128 @@ namespace SkeletalMeshTools
 			FSkinnedMeshChunk* SrcChunk = SrcChunks[SrcChunkIndex];
 			SrcChunk->OriginalSectionIndex = SrcChunkIndex;
 			int32 FirstChunkIndex = Chunks.Num();
-			for (int32 i = 0; i < SrcChunk->Indices.Num(); i += 3)
+			//Iterate Indice in the order of the shell patch
+			TArray<uint32> SrcChunkRemapIndicesIndex;
+			SrcChunkRemapIndicesIndex.Reserve(SrcChunk->Indices.Num());
+			TArray<PolygonShellsHelper::FPatchAndBoneInfluence>& PatchData = PatchDataPerSrcChunk[SrcChunkIndex];
+			TArray<int32>& IndiceToPatchIndex = IndiceToPatchIndexPerSrcChunk[SrcChunkIndex];
+			for (int32 PatchIndex = 0; PatchIndex < PatchData.Num(); ++PatchIndex)
 			{
-				// Find all bones needed by this triangle.
-				UniqueBones.Reset();
-				for (int32 Corner = 0; Corner < 3; Corner++)
+				if (!PatchData[PatchIndex].bIsParent)
+				{
+					continue;
+				}
+				SrcChunkRemapIndicesIndex.Reset();
+				PolygonShellsHelper::RecursiveFillRemapIndices(PatchData, PatchIndex, IndiceToPatchIndex, SrcChunkRemapIndicesIndex);
+
+				//Force adding a chunk since we want to control where we cut the model
+				int32 CurrentChunkIndex = FirstChunkIndex;
+
+				auto CreateChunk = [&SrcChunk, &FirstChunkIndex, &CurrentChunkIndex, &Chunks, &IndexMaps](FSkinnedMeshChunk** DestinationChunk)
+				{
+					(*DestinationChunk) = new FSkinnedMeshChunk();
+					CurrentChunkIndex = Chunks.Add(*DestinationChunk);
+					(*DestinationChunk)->MaterialIndex = SrcChunk->MaterialIndex;
+					(*DestinationChunk)->OriginalSectionIndex = SrcChunk->OriginalSectionIndex;
+					(*DestinationChunk)->ParentChunkSectionIndex = CurrentChunkIndex == FirstChunkIndex ? INDEX_NONE : FirstChunkIndex;
+					TArray<int32>& IndexMap = *new TArray<int32>();
+					IndexMaps.Add(&IndexMap);
+					IndexMap.AddUninitialized(SrcChunk->Vertices.Num());
+					FMemory::Memset(IndexMap.GetData(), 0xff, IndexMap.GetTypeSize()*IndexMap.Num());
+				};
+				
+				//Create a chunk
+				{
+					FSkinnedMeshChunk* DestChunk = NULL;
+					CreateChunk(&DestChunk);
+				}
+
+				//Add Indices to the chunk and add extra chunk only in case the patch use more bone then the maximum specified
+				for (int32 i = 0; i < SrcChunkRemapIndicesIndex.Num(); i += 3)
 				{
-					uint32 VertexIndex = SrcChunk->Indices[i + Corner];
-					FSoftSkinBuildVertex& V = SrcChunk->Vertices[VertexIndex];
-					for (int32 InfluenceIndex = 0; InfluenceIndex < MAX_TOTAL_INFLUENCES; InfluenceIndex++)
+					//We remap the iteration order to avoid cutting polygon shell
+					int32 IndiceIndex = SrcChunkRemapIndicesIndex[i];
+					// Find all bones needed by this triangle.
+					UniqueBones.Reset();
+					for (int32 Corner = 0; Corner < 3; Corner++)
 					{
-						if (V.InfluenceWeights[InfluenceIndex] > 0)
+						uint32 VertexIndex = SrcChunk->Indices[IndiceIndex + Corner];
+						FSoftSkinBuildVertex& V = SrcChunk->Vertices[VertexIndex];
+						for (int32 InfluenceIndex = 0; InfluenceIndex < MAX_TOTAL_INFLUENCES; InfluenceIndex++)
 						{
-							UniqueBones.AddUnique(V.InfluenceBones[InfluenceIndex]);
+							if (V.InfluenceWeights[InfluenceIndex] > 0)
+							{
+								UniqueBones.AddUnique(V.InfluenceBones[InfluenceIndex]);
+							}
+						}
+						//Add the alternate bones
+						TArray<FBoneIndexType>* AlternateBones = AlternateBoneIDs.Find(V.PointWedgeIdx);
+						if (AlternateBones)
+						{
+							for (int32 InfluenceIndex = 0; InfluenceIndex < AlternateBones->Num(); InfluenceIndex++)
+							{
+								UniqueBones.AddUnique((*AlternateBones)[InfluenceIndex]);
+							}
 						}
 					}
-					//Add the alternate bones
-					TArray<FBoneIndexType>* AlternateBones = AlternateBoneIDs.Find(V.PointWedgeIdx);
-					if (AlternateBones)
+
+					// Now find a chunk for them.
+					FSkinnedMeshChunk* DestChunk = NULL;
+					int32 DestChunkIndex = CurrentChunkIndex;
+					for (; DestChunkIndex < Chunks.Num(); ++DestChunkIndex)
 					{
-						for (int32 InfluenceIndex = 0; InfluenceIndex < AlternateBones->Num(); InfluenceIndex++)
+						TArray<FBoneIndexType>& BoneMap = Chunks[DestChunkIndex]->BoneMap;
+						int32 NumUniqueBones = 0;
+						for (int32 j = 0; j < UniqueBones.Num(); ++j)
 						{
-							UniqueBones.AddUnique((*AlternateBones)[InfluenceIndex]);
+							NumUniqueBones += (BoneMap.Contains(UniqueBones[j]) ? 0 : 1);
+						}
+						if (NumUniqueBones + BoneMap.Num() <= MaxBonesPerChunk)
+						{
+							DestChunk = Chunks[DestChunkIndex];
+							break;
 						}
 					}
-				}
 
-				// Now find a chunk for them.
-				FSkinnedMeshChunk* DestChunk = NULL;
-				int32 DestChunkIndex = FirstChunkIndex;
-				for (; DestChunkIndex < Chunks.Num(); ++DestChunkIndex)
-				{
-					TArray<FBoneIndexType>& BoneMap = Chunks[DestChunkIndex]->BoneMap;
-					int32 NumUniqueBones = 0;
-					for (int32 j = 0; j < UniqueBones.Num(); ++j)
+					// If no chunk was found, create one!
+					if (DestChunk == NULL)
 					{
-						NumUniqueBones += (BoneMap.Contains(UniqueBones[j]) ? 0 : 1);
+						CreateChunk(&DestChunk);
 					}
-					if (NumUniqueBones + BoneMap.Num() <= MaxBonesPerChunk)
+					TArray<int32>& IndexMap = IndexMaps[DestChunkIndex];
+
+					// Add the unique bones to this chunk's bone map.
+					for (int32 j = 0; j < UniqueBones.Num(); ++j)
 					{
-						DestChunk = Chunks[DestChunkIndex];
-						break;
+						DestChunk->BoneMap.AddUnique(UniqueBones[j]);
 					}
-				}
-
-				// If no chunk was found, create one!
-				if (DestChunk == NULL)
-				{
-					DestChunk = new FSkinnedMeshChunk();
-					int32 ChunkSectionIndex = Chunks.Add(DestChunk);
-					DestChunk->MaterialIndex = SrcChunk->MaterialIndex;
-					DestChunk->OriginalSectionIndex = SrcChunk->OriginalSectionIndex;
-					DestChunk->ParentChunkSectionIndex = ChunkSectionIndex == FirstChunkIndex ? INDEX_NONE : FirstChunkIndex;
-					TArray<int32>& IndexMap = *new TArray<int32>();
-					IndexMaps.Add(&IndexMap);
-					IndexMap.AddUninitialized(SrcChunk->Vertices.Num());
-					FMemory::Memset(IndexMap.GetData(),0xff,IndexMap.GetTypeSize()*IndexMap.Num());
-				}
-				TArray<int32>& IndexMap = IndexMaps[DestChunkIndex];
 
-				// Add the unique bones to this chunk's bone map.
-				for (int32 j = 0; j < UniqueBones.Num(); ++j)
-				{
-					DestChunk->BoneMap.AddUnique(UniqueBones[j]);
-				}
-
-				// For each vertex, add it to the chunk's arrays of vertices and indices.
-				for (int32 Corner = 0; Corner < 3; Corner++)
-				{
-					int32 VertexIndex = SrcChunk->Indices[i + Corner];
-					int32 DestIndex = IndexMap[VertexIndex];
-					if (DestIndex == INDEX_NONE)
+					// For each vertex, add it to the chunk's arrays of vertices and indices.
+					for (int32 Corner = 0; Corner < 3; Corner++)
 					{
-						DestIndex = DestChunk->Vertices.Add(SrcChunk->Vertices[VertexIndex]);
-						FSoftSkinBuildVertex& V = DestChunk->Vertices[DestIndex];
-						for (int32 InfluenceIndex = 0; InfluenceIndex < MAX_TOTAL_INFLUENCES; InfluenceIndex++)
+						int32 VertexIndex = SrcChunk->Indices[IndiceIndex + Corner];
+						int32 DestIndex = IndexMap[VertexIndex];
+						if (DestIndex == INDEX_NONE)
 						{
-							if (V.InfluenceWeights[InfluenceIndex] > 0)
+							DestIndex = DestChunk->Vertices.Add(SrcChunk->Vertices[VertexIndex]);
+							FSoftSkinBuildVertex& V = DestChunk->Vertices[DestIndex];
+							for (int32 InfluenceIndex = 0; InfluenceIndex < MAX_TOTAL_INFLUENCES; InfluenceIndex++)
 							{
-								int32 MappedIndex = DestChunk->BoneMap.Find(V.InfluenceBones[InfluenceIndex]);
-								check(DestChunk->BoneMap.IsValidIndex(MappedIndex));
-								V.InfluenceBones[InfluenceIndex] = MappedIndex;
+								if (V.InfluenceWeights[InfluenceIndex] > 0)
+								{
+									int32 MappedIndex = DestChunk->BoneMap.Find(V.InfluenceBones[InfluenceIndex]);
+									check(DestChunk->BoneMap.IsValidIndex(MappedIndex));
+									V.InfluenceBones[InfluenceIndex] = MappedIndex;
+								}
 							}
+							IndexMap[VertexIndex] = DestIndex;
 						}
-						IndexMap[VertexIndex] = DestIndex;
+						DestChunk->Indices.Add(DestIndex);
 					}
-					DestChunk->Indices.Add(DestIndex);
 				}
 			}
 
+
 			// Source chunks are no longer needed.
 			delete SrcChunks[SrcChunkIndex];
 			SrcChunks[SrcChunkIndex] = NULL;
diff --git a/Engine/Source/Runtime/Engine/Private/SkeletalMeshRenderData.cpp b/Engine/Source/Runtime/Engine/Private/SkeletalMeshRenderData.cpp
--- a/Engine/Source/Runtime/Engine/Private/SkeletalMeshRenderData.cpp
+++ b/Engine/Source/Runtime/Engine/Private/SkeletalMeshRenderData.cpp
@@ -59,7 +59,7 @@ static void SerializeLODInfoForDDC(USkeletalMesh* SkeletalMesh, FString& KeySuff
 // differences, etc.) replace the version GUID below with a new one.
 // In case of merge conflicts with DDC versions, you MUST generate a new GUID
 // and set this new GUID as the version.                                       
-#define SKELETALMESH_DERIVEDDATA_VER TEXT("29B218D50FE444AF8EFC8E4234F5D4E1")
+#define SKELETALMESH_DERIVEDDATA_VER TEXT("1DE2DE4E9ADC40C494436E7729B14D5F")
 
 static const FString& GetSkeletalMeshDerivedDataVersion()
 {
