Optimize the polygonshell chunking algorithm
#rnx

diff --git a/Engine/Source/Developer/MeshUtilities/Private/SkeletalMeshTools.cpp b/Engine/Source/Developer/MeshUtilities/Private/SkeletalMeshTools.cpp
--- a/Engine/Source/Developer/MeshUtilities/Private/SkeletalMeshTools.cpp
+++ b/Engine/Source/Developer/MeshUtilities/Private/SkeletalMeshTools.cpp
@@ -271,13 +271,12 @@ namespace SkeletalMeshTools
 		//Fill FaceIndexToPatchIndex so every triangle knows its unique island patch index.
 		//Each island patch have is fill with connected vertexinstance where position, NTBs. UVs and colors are nearly equal.
 		//@Param bConnectByEdge: If true we need at least 2 vertex index (one edge) to connect 2 triangles. If false we just need one vertex index (bowtie)
-		void FillPolygonPatch(const TArray<uint32>& Indices, const TArray<FSoftSkinBuildVertex>& Vertices, const TMap<uint32, TArray<FBoneIndexType>>& AlternateBoneIDs, TArray<FPatchAndBoneInfluence>& PatchData, TArray<int32>& IndiceToPatchIndex, const int32 MaxBonesPerChunk, const bool bConnectByEdge)
+		void FillPolygonPatch(const TArray<uint32>& Indices, const TArray<FSoftSkinBuildVertex>& Vertices, const TMap<uint32, TArray<FBoneIndexType>>& AlternateBoneIDs, TArray<FPatchAndBoneInfluence>& PatchData, TArray<TArray<uint32>>& PatchIndexToIndices, const int32 MaxBonesPerChunk, const bool bConnectByEdge)
 		{
 			const int32 NumIndice = Indices.Num();
-			check(IndiceToPatchIndex.Num() == NumIndice);
 			const int32 NumFace = NumIndice / 3;
 			int32 PatchIndex = 0;
-
+			
 			//Store a map containing connected faces for each vertex index
 			TMap<int32, TArray<int32>> VertexIndexToAdjacentFaces;
 			VertexIndexToAdjacentFaces.Reserve(Vertices.Num());
@@ -299,6 +298,10 @@ namespace SkeletalMeshTools
 
 			TArray<int32> TriangleQueue;
 			TriangleQueue.Reserve(100);
+			//Allocate an array and use it to retrieve the data, we do not know the number of indices per patch so it prevent us doing a huge reserve per patch
+			//Simply copy the result in PatchIndexToIndices when we finish gathering the patch data.
+			TArray<uint32> AllocatedPatchIndexToIndices;
+			AllocatedPatchIndexToIndices.Reserve(NumIndice);
 			for (int32 FaceIndex = 0; FaceIndex < NumFace; ++FaceIndex)
 			{
 				//Skip already added faces
@@ -306,6 +309,7 @@ namespace SkeletalMeshTools
 				{
 					continue;
 				}
+				AllocatedPatchIndexToIndices.Reset();
 
 				//Add all the faces connected to the current face index
 				TriangleQueue.Reset();
@@ -318,8 +322,7 @@ namespace SkeletalMeshTools
 					for (int32 Corner = 0; Corner < 3; Corner++)
 					{
 						const int32 IndiceIndex = IndiceOffset + Corner;
-						checkSlow(IndiceToPatchIndex.IsValidIndex(IndiceIndex));
-						IndiceToPatchIndex[IndiceIndex] = PatchIndex;
+						AllocatedPatchIndexToIndices.Add(IndiceIndex);
 						checkSlow(Indices.IsValidIndex(IndiceIndex));
 						const int32 VertexIndex = Indices[IndiceIndex];
 						checkSlow(Vertices.IsValidIndex(VertexIndex));
@@ -354,25 +357,24 @@ namespace SkeletalMeshTools
 
 					AddAdjacentFace(Indices, Vertices, FaceAdded, VertexIndexToAdjacentFaces, CurrentTriangleIndex, TriangleQueue, bConnectByEdge);
 				}
+
+				//This is a new patch create the data and append the patch result remap
+				check(!PatchIndexToIndices.IsValidIndex(PatchIndex));
+				PatchIndexToIndices.AddDefaulted();
+				check(PatchIndexToIndices.IsValidIndex(PatchIndex));
+				PatchIndexToIndices[PatchIndex].Append(AllocatedPatchIndexToIndices);
 				PatchIndex++;
 			}
 		}
 
-		void RecursiveFillRemapIndices(const TArray<FPatchAndBoneInfluence>& PatchData, const int32 PatchIndex, const TArray<int32>& IndiceToPatchIndex, TArray<uint32>& SrcChunkRemapIndicesIndex)
+		void RecursiveFillRemapIndices(const TArray<FPatchAndBoneInfluence>& PatchData, const int32 PatchIndex, const TArray<TArray<uint32>>& PatchIndexToIndices, TArray<uint32>& SrcChunkRemapIndicesIndex)
 		{
-			for (int32 IndiceIndex = 0; IndiceIndex < IndiceToPatchIndex.Num(); ++IndiceIndex)
-			{
-				int32 IndicePatchIndex = IndiceToPatchIndex[IndiceIndex];
-				if (PatchIndex == IndicePatchIndex)
-				{
-					SrcChunkRemapIndicesIndex.Add(IndiceIndex);
-				}
-			}
+			SrcChunkRemapIndicesIndex.Append(PatchIndexToIndices[PatchIndex]);
 			checkSlow(PatchData.IsValidIndex(PatchIndex));
 			//Do the child patch to chunk with
 			for (int32 SubPatchIndex = 0; SubPatchIndex < PatchData[PatchIndex].PatchToChunkWith.Num(); ++SubPatchIndex)
 			{
-				RecursiveFillRemapIndices(PatchData, PatchData[PatchIndex].PatchToChunkWith[SubPatchIndex], IndiceToPatchIndex, SrcChunkRemapIndicesIndex);
+				RecursiveFillRemapIndices(PatchData, PatchData[PatchIndex].PatchToChunkWith[SubPatchIndex], PatchIndexToIndices, SrcChunkRemapIndicesIndex);
 			}
 		}
 
@@ -432,7 +434,7 @@ namespace SkeletalMeshTools
 		SrcChunks.Sort(FCompareSkinnedMeshChunk());
 
 		TMap<int32, TArray<PolygonShellsHelper::FPatchAndBoneInfluence>> PatchDataPerSrcChunk;
-		TMap<int32, TArray<int32>> IndiceToPatchIndexPerSrcChunk;
+		TMap<int32, TArray<TArray<uint32>>> PatchIndexToIndicesPerSrcChunk;
 		
 		//Find the shells inside chunks
 		for (int32 ChunkIndex = 0; ChunkIndex < SrcChunks.Num(); ++ChunkIndex)
@@ -441,11 +443,10 @@ namespace SkeletalMeshTools
 			TArray<uint32>& Indices = ChunkToShell->Indices;
 			TArray<FSoftSkinBuildVertex>& Vertices = ChunkToShell->Vertices;
 			TArray<PolygonShellsHelper::FPatchAndBoneInfluence>& PatchData = PatchDataPerSrcChunk.Add(ChunkIndex);
-			TArray<int32>& IndiceToPatchIndex = IndiceToPatchIndexPerSrcChunk.Add(ChunkIndex);
-			IndiceToPatchIndex.AddZeroed(Indices.Num());
+			TArray<TArray<uint32>>& PatchIndexToIndices = PatchIndexToIndicesPerSrcChunk.Add(ChunkIndex);
 			//We need edge connection (2 similar vertex )
 			const bool bConnectByEdge = true;
-			PolygonShellsHelper::FillPolygonPatch(Indices, Vertices, AlternateBoneIDs, PatchData, IndiceToPatchIndex, MaxBonesPerChunk, bConnectByEdge);
+			PolygonShellsHelper::FillPolygonPatch(Indices, Vertices, AlternateBoneIDs, PatchData, PatchIndexToIndices, MaxBonesPerChunk, bConnectByEdge);
 		}
 
 		for (int32 SrcChunkIndex = 0; SrcChunkIndex < SrcChunks.Num(); ++SrcChunkIndex)
@@ -477,7 +478,8 @@ namespace SkeletalMeshTools
 			TArray<uint32> SrcChunkRemapIndicesIndex;
 			SrcChunkRemapIndicesIndex.Reserve(SrcChunk->Indices.Num());
 			TArray<PolygonShellsHelper::FPatchAndBoneInfluence>& PatchData = PatchDataPerSrcChunk[SrcChunkIndex];
-			TArray<int32>& IndiceToPatchIndex = IndiceToPatchIndexPerSrcChunk[SrcChunkIndex];
+			const TArray<TArray<uint32>>& PatchIndexToIndices = PatchIndexToIndicesPerSrcChunk[SrcChunkIndex];
+
 			for (int32 PatchIndex = 0; PatchIndex < PatchData.Num(); ++PatchIndex)
 			{
 				if (!PatchData[PatchIndex].bIsParent)
@@ -485,7 +487,7 @@ namespace SkeletalMeshTools
 					continue;
 				}
 				SrcChunkRemapIndicesIndex.Reset();
-				PolygonShellsHelper::RecursiveFillRemapIndices(PatchData, PatchIndex, IndiceToPatchIndex, SrcChunkRemapIndicesIndex);
+				PolygonShellsHelper::RecursiveFillRemapIndices(PatchData, PatchIndex, PatchIndexToIndices, SrcChunkRemapIndicesIndex);
 
 				//Force adding a chunk since we want to control where we cut the model
 				int32 CurrentChunkIndex = FirstChunkIndex;
