#UE4 Disabled a warning about truncating the number of influences for now.

diff --git a/Engine/Source/Developer/MeshUtilities/Private/MeshUtilities.cpp b/Engine/Source/Developer/MeshUtilities/Private/MeshUtilities.cpp
--- a/Engine/Source/Developer/MeshUtilities/Private/MeshUtilities.cpp
+++ b/Engine/Source/Developer/MeshUtilities/Private/MeshUtilities.cpp
@@ -4067,7 +4067,7 @@ public:
 					InfluenceCount = FMath::Min<uint32>(InfluenceCount, MAX_TOTAL_INFLUENCES);
 					if (InfluenceCount > EXTRA_BONE_INFLUENCES && !FGPUBaseSkinVertexFactory::UseUnlimitedBoneInfluences(InfluenceCount))
 					{
-						UE_LOG(LogSkeletalMesh, Warning, TEXT("Skeletal mesh of %d bone influences requires unlimited bone influence mode on. Influence truncated to %d."), InfluenceCount, EXTRA_BONE_INFLUENCES);
+						//UE_LOG(LogSkeletalMesh, Warning, TEXT("Skeletal mesh of %d bone influences requires unlimited bone influence mode on. Influence truncated to %d."), InfluenceCount, EXTRA_BONE_INFLUENCES);
 						InfluenceCount = EXTRA_BONE_INFLUENCES;
 					}
 
