Generate mesh distance fields and Lumen card representation for Nanite meshes directly from the source mesh data. Nanite replaces static mesh data with a coarse mesh representation, which sometimes has holes in it or is not accurate enough. Building from the source data works around that issue and gets up the highest possible SDF/card quality. Unfortunately it also regresses static mesh import times by ~3-5% until we don't get async static mesh builds.

Mesh build numbers for Nanite meshes:
1701216 tris:
Static mesh build: 11.2s
Static mesh build for the SDF async task (added in this CL): 2.3s
Async SDF build: 0.7s
Async card rep build: 0.15s

5524435 tris:
Static mesh build: 36.4s
Static mesh build for the SDF async task: 6.2s
Async SDF build: 1.2s
Async card rep build: 0.5s

1146238 tris
Static mesh build: 10.5s
Static mesh build for the SDF async task: 1.5s
Async SDF build: 0.6s
Async card rep build: 0.1s

diff --git a/Engine/Source/Developer/MeshBuilder/Private/MeshBuilderModule.cpp b/Engine/Source/Developer/MeshBuilder/Private/MeshBuilderModule.cpp
--- a/Engine/Source/Developer/MeshBuilder/Private/MeshBuilderModule.cpp
+++ b/Engine/Source/Developer/MeshBuilder/Private/MeshBuilderModule.cpp
@@ -33,6 +33,7 @@ public:
 		TArray< FStaticMeshBuildVertex >& Verts,
 		TArray< uint32 >& Indexes,
 		FStaticMeshSectionArray& Sections,
+		bool bBuildOnlyPosition,
 		uint32& NumTexCoords,
 		bool& bHasColors ) override;
 
@@ -60,6 +61,7 @@ bool FMeshBuilderModule::BuildMesh(
 	TArray< FStaticMeshBuildVertex >& Verts,
 	TArray< uint32 >& Indexes,
 	FStaticMeshSectionArray& Sections,
+	bool bBuildOnlyPosition,
 	uint32& NumTexCoords,
 	bool& bHasColors )
 {
@@ -67,7 +69,7 @@ bool FMeshBuilderModule::BuildMesh(
 	if( StaticMesh )
 	{
 		//Call the static mesh builder
-		return FStaticMeshBuilder().Build( StaticMesh, Verts, Indexes, Sections, NumTexCoords, bHasColors );
+		return FStaticMeshBuilder().Build( StaticMesh, Verts, Indexes, Sections, bBuildOnlyPosition, NumTexCoords, bHasColors );
 	}
 	return false;
 }
diff --git a/Engine/Source/Developer/MeshBuilder/Private/MeshDescriptionHelper.cpp b/Engine/Source/Developer/MeshBuilder/Private/MeshDescriptionHelper.cpp
--- a/Engine/Source/Developer/MeshBuilder/Private/MeshDescriptionHelper.cpp
+++ b/Engine/Source/Developer/MeshBuilder/Private/MeshDescriptionHelper.cpp
@@ -30,7 +30,7 @@ FMeshDescriptionHelper::FMeshDescriptionHelper(FMeshBuildSettings* InBuildSettin
 {
 }
 
-void FMeshDescriptionHelper::SetupRenderMeshDescription(UObject* Owner, FMeshDescription& RenderMeshDescription)
+void FMeshDescriptionHelper::SetupRenderMeshDescription(UObject* Owner, FMeshDescription& RenderMeshDescription, bool bBuildOnlyPosition)
 {
 	TRACE_CPUPROFILER_EVENT_SCOPE(FMeshDescriptionHelper::GetRenderMeshDescription);
 
@@ -54,6 +54,7 @@ void FMeshDescriptionHelper::SetupRenderMeshDescription(UObject* Owner, FMeshDes
 	FStaticMeshOperations::FindOverlappingCorners(OverlappingCorners, RenderMeshDescription, ComparisonThreshold);
 
 	// Compute any missing normals or tangents.
+	if (!bBuildOnlyPosition)
 	{
 		// Static meshes always blend normals of overlapping corners.
 		EComputeNTBsFlags ComputeNTBsOptions = EComputeNTBsFlags::BlendOverlappingNormals;
@@ -66,7 +67,7 @@ void FMeshDescriptionHelper::SetupRenderMeshDescription(UObject* Owner, FMeshDes
 		FStaticMeshOperations::ComputeTangentsAndNormals(RenderMeshDescription, ComputeNTBsOptions);
 	}
 
-	if (BuildSettings->bGenerateLightmapUVs && VertexInstanceArray.Num() > 0)
+	if (BuildSettings->bGenerateLightmapUVs && !bBuildOnlyPosition && VertexInstanceArray.Num() > 0)
 	{
 		TVertexInstanceAttributesRef<FVector2D> VertexInstanceUVs = Attributes.GetVertexInstanceUVs();
 		int32 NumIndices = VertexInstanceUVs.GetNumChannels();
diff --git a/Engine/Source/Developer/MeshBuilder/Private/MeshDescriptionHelper.h b/Engine/Source/Developer/MeshBuilder/Private/MeshDescriptionHelper.h
--- a/Engine/Source/Developer/MeshBuilder/Private/MeshDescriptionHelper.h
+++ b/Engine/Source/Developer/MeshBuilder/Private/MeshDescriptionHelper.h
@@ -21,7 +21,7 @@ public:
 	FMeshDescriptionHelper(FMeshBuildSettings* InBuildSettings);
 
 	//Build a render mesh description with the BuildSettings. This will update the RenderMeshDescription in place
-	void SetupRenderMeshDescription(UObject* Owner, FMeshDescription& RenderMeshDescription);
+	void SetupRenderMeshDescription(UObject* Owner, FMeshDescription& RenderMeshDescription, bool bBuildOnlyPosition = false);
 
 	void ReduceLOD(const FMeshDescription& BaseMesh, FMeshDescription& DestMesh, const struct FMeshReductionSettings& ReductionSettings, const FOverlappingCorners& InOverlappingCorners, float &OutMaxDeviation);
 
diff --git a/Engine/Source/Developer/MeshBuilder/Private/StaticMeshBuilder.cpp b/Engine/Source/Developer/MeshBuilder/Private/StaticMeshBuilder.cpp
--- a/Engine/Source/Developer/MeshBuilder/Private/StaticMeshBuilder.cpp
+++ b/Engine/Source/Developer/MeshBuilder/Private/StaticMeshBuilder.cpp
@@ -376,6 +376,7 @@ bool FStaticMeshBuilder::Build(
 	TArray< FStaticMeshBuildVertex >& Verts,
 	TArray< uint32 >& Indexes,
 	FStaticMeshSectionArray& Sections,
+	bool bBuildOnlyPosition,
 	uint32& NumTexCoords,
 	bool& bHasColors )
 {
@@ -399,7 +400,7 @@ bool FStaticMeshBuilder::Build(
 	{
 		TRACE_CPUPROFILER_EVENT_SCOPE_TEXT(TEXT("FStaticMeshBuilder::Build::Mappings"));
 		MeshDescription = *OriginalMeshDescription;
-		MeshDescriptionHelper.SetupRenderMeshDescription( StaticMesh, MeshDescription);
+		MeshDescriptionHelper.SetupRenderMeshDescription(StaticMesh, MeshDescription, bBuildOnlyPosition);
 	}
 	else
 	{
diff --git a/Engine/Source/Developer/MeshBuilder/Public/IMeshBuilderModule.h b/Engine/Source/Developer/MeshBuilder/Public/IMeshBuilderModule.h
--- a/Engine/Source/Developer/MeshBuilder/Public/IMeshBuilderModule.h
+++ b/Engine/Source/Developer/MeshBuilder/Public/IMeshBuilderModule.h
@@ -36,6 +36,7 @@ public:
 		TArray< FStaticMeshBuildVertex >& Verts,
 		TArray< uint32 >& Indexes,
 		FStaticMeshSectionArray& Sections,
+		bool bBuildOnlyPosition,
 		uint32& NumTexCoords,
 		bool& bHasColors )
 	{
diff --git a/Engine/Source/Developer/MeshBuilder/Public/MeshBuilder.h b/Engine/Source/Developer/MeshBuilder/Public/MeshBuilder.h
--- a/Engine/Source/Developer/MeshBuilder/Public/MeshBuilder.h
+++ b/Engine/Source/Developer/MeshBuilder/Public/MeshBuilder.h
@@ -31,6 +31,7 @@ public:
 		TArray< FStaticMeshBuildVertex >& Verts,
 		TArray< uint32 >& Indexes,
 		FStaticMeshSectionArray& Sections,
+		bool bBuildOnlyPosition,
 		uint32& NumTexCoords,
 		bool& bHasColors ) = 0;
 
diff --git a/Engine/Source/Developer/MeshBuilder/Public/SkeletalMeshBuilder.h b/Engine/Source/Developer/MeshBuilder/Public/SkeletalMeshBuilder.h
--- a/Engine/Source/Developer/MeshBuilder/Public/SkeletalMeshBuilder.h
+++ b/Engine/Source/Developer/MeshBuilder/Public/SkeletalMeshBuilder.h
@@ -30,6 +30,7 @@ public:
 		TArray< FStaticMeshBuildVertex >& Verts,
 		TArray< uint32 >& Indexes,
 		FStaticMeshSectionArray& Sections,
+		bool bBuildOnlyPosition,
 		uint32& NumTexCoords,
 		bool& bHasColors) override
 	{
diff --git a/Engine/Source/Developer/MeshBuilder/Public/StaticMeshBuilder.h b/Engine/Source/Developer/MeshBuilder/Public/StaticMeshBuilder.h
--- a/Engine/Source/Developer/MeshBuilder/Public/StaticMeshBuilder.h
+++ b/Engine/Source/Developer/MeshBuilder/Public/StaticMeshBuilder.h
@@ -33,6 +33,7 @@ public:
 		TArray< FStaticMeshBuildVertex >& Verts,
 		TArray< uint32 >& Indexes,
 		FStaticMeshSectionArray& Sections,
+		bool bBuildOnlyPosition,
 		uint32& NumTexCoords,
 		bool& bHasColors ) override;
 
diff --git a/Engine/Source/Developer/MeshUtilities/Private/MeshCardRepresentationUtilities.cpp b/Engine/Source/Developer/MeshUtilities/Private/MeshCardRepresentationUtilities.cpp
--- a/Engine/Source/Developer/MeshUtilities/Private/MeshCardRepresentationUtilities.cpp
+++ b/Engine/Source/Developer/MeshUtilities/Private/MeshCardRepresentationUtilities.cpp
@@ -21,11 +21,14 @@ bool SetupEmbreeScene(
 	FString MeshName,
 	int32 PrimitiveRangeStart, 
 	int32 PrimitiveRangeEnd, 
-	const FInputPrimitiveIds& InputPrimitiveIds, 
+	const FInputPrimitiveIds& InputPrimitiveIds,
+	const FSourceMeshDataForDerivedDataTask& SourceMeshData,
 	const FStaticMeshLODResources& LODModel,
 	RTCDevice EmbreeDevice,
 	RTCScene& EmbreeScene)
 {
+	const int32 NumVertices = SourceMeshData.IsValid() ? SourceMeshData.Vertices.Num() : LODModel.VertexBuffers.PositionVertexBuffer.GetNumVertices();
+
 	EmbreeScene = rtcDeviceNewScene(EmbreeDevice, RTC_SCENE_STATIC, RTC_INTERSECT1);
 			
 	RTCError ReturnErrorNewScene = rtcDeviceGetError(EmbreeDevice);
@@ -36,14 +39,12 @@ bool SetupEmbreeScene(
 		return false;
 	}
 
-	const FPositionVertexBuffer& PositionVertexBuffer = LODModel.VertexBuffers.PositionVertexBuffer;
-	FIndexArrayView Indices = LODModel.IndexBuffer.GetArrayView();
 	const int32 NumTriangles = PrimitiveRangeEnd - PrimitiveRangeStart;
 
 	FVector4* EmbreeVertices = NULL;
 	int32* EmbreeIndices = NULL;
 
-	uint32 GeomID = rtcNewTriangleMesh(EmbreeScene, RTC_GEOMETRY_STATIC, NumTriangles, PositionVertexBuffer.GetNumVertices());
+	uint32 GeomID = rtcNewTriangleMesh(EmbreeScene, RTC_GEOMETRY_STATIC, NumTriangles, NumVertices);
 
 	EmbreeVertices = (FVector4*)rtcMapBuffer(EmbreeScene, GeomID, RTC_VERTEX_BUFFER);
 	EmbreeIndices = (int32*)rtcMapBuffer(EmbreeScene, GeomID, RTC_INDEX_BUFFER);
@@ -51,13 +52,31 @@ bool SetupEmbreeScene(
 	for (int32 PrimitiveIndex = PrimitiveRangeStart; PrimitiveIndex < PrimitiveRangeEnd; PrimitiveIndex++)
 	{
 		int32 TriangleIndex = InputPrimitiveIds[PrimitiveIndex];
-		int32 I0 = Indices[TriangleIndex * 3 + 0];
-		int32 I1 = Indices[TriangleIndex * 3 + 1];
-		int32 I2 = Indices[TriangleIndex * 3 + 2];
 
-		FVector V0 = PositionVertexBuffer.VertexPosition(I0);
-		FVector V1 = PositionVertexBuffer.VertexPosition(I1);
-		FVector V2 = PositionVertexBuffer.VertexPosition(I2);
+		int32 I0, I1, I2;
+		FVector V0, V1, V2;
+
+		if (SourceMeshData.IsValid())
+		{
+			I0 = SourceMeshData.Indices[TriangleIndex * 3 + 0];
+			I1 = SourceMeshData.Indices[TriangleIndex * 3 + 1];
+			I2 = SourceMeshData.Indices[TriangleIndex * 3 + 2];
+
+			V0 = SourceMeshData.Vertices[I0].Position;
+			V1 = SourceMeshData.Vertices[I1].Position;
+			V2 = SourceMeshData.Vertices[I2].Position;
+		}
+		else
+		{
+			const FIndexArrayView Indices = LODModel.IndexBuffer.GetArrayView();
+			I0 = Indices[TriangleIndex * 3 + 0];
+			I1 = Indices[TriangleIndex * 3 + 1];
+			I2 = Indices[TriangleIndex * 3 + 2];
+
+			V0 = LODModel.VertexBuffers.PositionVertexBuffer.VertexPosition(I0);
+			V1 = LODModel.VertexBuffers.PositionVertexBuffer.VertexPosition(I1);
+			V2 = LODModel.VertexBuffers.PositionVertexBuffer.VertexPosition(I2);
+		}
 
 		EmbreeIndices[TriangleIndex * 3 + 0] = I0;
 		EmbreeIndices[TriangleIndex * 3 + 1] = I1;
@@ -110,14 +129,12 @@ public:
 	const FString& MeshName;
 	RTCScene FullMeshEmbreeScene;
 	RTCDevice EmbreeDevice;
-	const FStaticMeshLODResources& LODModel;
 	FCardRepresentationData& OutData;
 
-	FGenerateCardMeshContext(const FString& InMeshName, RTCScene InEmbreeScene, RTCDevice InEmbreeDevice, const FStaticMeshLODResources& InLODModel, FCardRepresentationData& InOutData) :
+	FGenerateCardMeshContext(const FString& InMeshName, RTCScene InEmbreeScene, RTCDevice InEmbreeDevice, FCardRepresentationData& InOutData) :
 		MeshName(InMeshName),
 		FullMeshEmbreeScene(InEmbreeScene),
 		EmbreeDevice(InEmbreeDevice),
-		LODModel(InLODModel),
 		OutData(InOutData)
 	{}
 };
@@ -517,6 +534,7 @@ void BuildCubeMapTree(const FBox& CardBounds, const FGenerateCardMeshContext& Co
 
 bool FMeshUtilities::GenerateCardRepresentationData(
 	FString MeshName,
+	const FSourceMeshDataForDerivedDataTask& SourceMeshData,
 	const FStaticMeshLODResources& LODModel,
 	class FQueuedThreadPool& ThreadPool,
 	const FBoxSphereBounds& Bounds,
@@ -535,9 +553,8 @@ bool FMeshUtilities::GenerateCardRepresentationData(
 		return false;
 	}
 
-	const FPositionVertexBuffer& PositionVertexBuffer = LODModel.VertexBuffers.PositionVertexBuffer;
-	FIndexArrayView Indices = LODModel.IndexBuffer.GetArrayView();
-	const int32 NumTriangles = Indices.Num() / 3;
+	const int32 NumIndices = SourceMeshData.IsValid() ? SourceMeshData.Indices.Num() : LODModel.IndexBuffer.GetNumIndices();
+	const int32 NumTriangles = NumIndices / 3;
 
 	FInputPrimitiveIds InputPrimitiveIds;
 	InputPrimitiveIds.Reserve(NumTriangles);
@@ -549,14 +566,14 @@ bool FMeshUtilities::GenerateCardRepresentationData(
 
 	RTCScene EmbreeScene = nullptr;
 
-	const bool bEmbreeSetup = SetupEmbreeScene(MeshName, 0, NumTriangles, InputPrimitiveIds, LODModel, EmbreeDevice, EmbreeScene);
+	const bool bEmbreeSetup = SetupEmbreeScene(MeshName, 0, NumTriangles, InputPrimitiveIds, SourceMeshData, LODModel, EmbreeDevice, EmbreeScene);
 
 	if (!bEmbreeSetup)
 	{
 		return false;
 	}
 
-	FGenerateCardMeshContext Context(MeshName, EmbreeScene, EmbreeDevice, LODModel, OutData);
+	FGenerateCardMeshContext Context(MeshName, EmbreeScene, EmbreeDevice, OutData);
 
 	BuildCubeMapTree(Bounds.GetBox(), Context, OutData);
 
diff --git a/Engine/Source/Developer/MeshUtilities/Private/MeshDistanceFieldUtilities.cpp b/Engine/Source/Developer/MeshUtilities/Private/MeshDistanceFieldUtilities.cpp
--- a/Engine/Source/Developer/MeshUtilities/Private/MeshDistanceFieldUtilities.cpp
+++ b/Engine/Source/Developer/MeshUtilities/Private/MeshDistanceFieldUtilities.cpp
@@ -319,6 +319,7 @@ void FMeshDistanceFieldAsyncTask::DoWork()
 
 void FMeshUtilities::GenerateSignedDistanceFieldVolumeData(
 	FString MeshName,
+	const FSourceMeshDataForDerivedDataTask& SourceMeshData,
 	const FStaticMeshLODResources& LODModel,
 	class FQueuedThreadPool& ThreadPool,
 	const TArray<EBlendMode>& MaterialBlendModes,
@@ -330,8 +331,11 @@ void FMeshUtilities::GenerateSignedDistanceFieldVolumeData(
 	if (DistanceFieldResolutionScale > 0)
 	{
 		const double StartTime = FPlatformTime::Seconds();
-		const FPositionVertexBuffer& PositionVertexBuffer = LODModel.VertexBuffers.PositionVertexBuffer;
-		FIndexArrayView Indices = LODModel.IndexBuffer.GetArrayView();
+
+		const int32 NumIndices = SourceMeshData.IsValid() ? SourceMeshData.Indices.Num() : LODModel.IndexBuffer.GetNumIndices();
+		const int32 NumVertices = SourceMeshData.IsValid() ? SourceMeshData.Vertices.Num() : LODModel.VertexBuffers.PositionVertexBuffer.GetNumVertices();
+		const FStaticMeshSectionArray& Sections = SourceMeshData.IsValid() ? SourceMeshData.Sections : LODModel.Sections;
+
 		TArray<FkDOPBuildCollisionTriangle<uint32> > BuildTriangles;
 
 		RTCDevice EmbreeDevice = NULL;
@@ -376,13 +380,34 @@ void FMeshUtilities::GenerateSignedDistanceFieldVolumeData(
 			&& Bounds.Origin.Z + Bounds.BoxExtent.Z > -KINDA_SMALL_NUMBER;
 
 		TArray<int32> FilteredTriangles;
-		FilteredTriangles.Empty(Indices.Num() / 3);
+		FilteredTriangles.Empty(NumIndices / 3);
 
-		for (int32 i = 0; i < Indices.Num(); i += 3)
+		for (int32 i = 0; i < NumIndices; i += 3)
 		{
-			FVector V0 = PositionVertexBuffer.VertexPosition(Indices[i + 0]);
-			FVector V1 = PositionVertexBuffer.VertexPosition(Indices[i + 1]);
-			FVector V2 = PositionVertexBuffer.VertexPosition(Indices[i + 2]);
+			int32 I0, I1, I2;
+			FVector V0, V1, V2;
+
+			if (SourceMeshData.IsValid())
+			{
+				I0 = SourceMeshData.Indices[i + 0];
+				I1 = SourceMeshData.Indices[i + 1];
+				I2 = SourceMeshData.Indices[i + 2];
+
+				V0 = SourceMeshData.Vertices[I0].Position;
+				V1 = SourceMeshData.Vertices[I1].Position;
+				V2 = SourceMeshData.Vertices[I2].Position;
+			}
+			else
+			{
+				const FIndexArrayView Indices = LODModel.IndexBuffer.GetArrayView();
+				I0 = Indices[i + 0];
+				I1 = Indices[i + 1];
+				I2 = Indices[i + 2];
+
+				V0 = LODModel.VertexBuffers.PositionVertexBuffer.VertexPosition(I0);
+				V1 = LODModel.VertexBuffers.PositionVertexBuffer.VertexPosition(I1);
+				V2 = LODModel.VertexBuffers.PositionVertexBuffer.VertexPosition(I2);
+			}
 
 			if (bMeshWasPlane)
 			{
@@ -399,9 +424,9 @@ void FMeshUtilities::GenerateSignedDistanceFieldVolumeData(
 			{
 				bool bTriangleIsOpaqueOrMasked = false;
 
-				for (int32 SectionIndex = 0; SectionIndex < LODModel.Sections.Num(); SectionIndex++)
+				for (int32 SectionIndex = 0; SectionIndex < Sections.Num(); SectionIndex++)
 				{
-					const FStaticMeshSection& Section = LODModel.Sections[SectionIndex];
+					const FStaticMeshSection& Section = Sections[SectionIndex];
 
 					if ((uint32)i >= Section.FirstIndex && (uint32)i < Section.FirstIndex + Section.NumTriangles * 3)
 					{
@@ -430,7 +455,7 @@ void FMeshUtilities::GenerateSignedDistanceFieldVolumeData(
 
 		if (bUseEmbree)
 		{
-			GeomID = rtcNewTriangleMesh(EmbreeScene, RTC_GEOMETRY_STATIC, FilteredTriangles.Num(), PositionVertexBuffer.GetNumVertices());
+			GeomID = rtcNewTriangleMesh(EmbreeScene, RTC_GEOMETRY_STATIC, FilteredTriangles.Num(), NumVertices);
 
 			rtcSetIntersectionFilterFunction(EmbreeScene, GeomID, EmbreeFilterFunc);
 			rtcSetOcclusionFilterFunction(EmbreeScene, GeomID, EmbreeFilterFunc);
@@ -445,13 +470,30 @@ void FMeshUtilities::GenerateSignedDistanceFieldVolumeData(
 
 		for (int32 TriangleIndex = 0; TriangleIndex < FilteredTriangles.Num(); TriangleIndex++)
 		{
-			int32 I0 = Indices[TriangleIndex * 3 + 0];
-			int32 I1 = Indices[TriangleIndex * 3 + 1];
-			int32 I2 = Indices[TriangleIndex * 3 + 2];
+			int32 I0, I1, I2;
+			FVector V0, V1, V2;
+
+			if (SourceMeshData.IsValid())
+			{
+				I0 = SourceMeshData.Indices[TriangleIndex * 3 + 0];
+				I1 = SourceMeshData.Indices[TriangleIndex * 3 + 1];
+				I2 = SourceMeshData.Indices[TriangleIndex * 3 + 2];
 
-			FVector V0 = PositionVertexBuffer.VertexPosition(I0);
-			FVector V1 = PositionVertexBuffer.VertexPosition(I1);
-			FVector V2 = PositionVertexBuffer.VertexPosition(I2);
+				V0 = SourceMeshData.Vertices[I0].Position;
+				V1 = SourceMeshData.Vertices[I1].Position;
+				V2 = SourceMeshData.Vertices[I2].Position;
+			}
+			else
+			{
+				const FIndexArrayView Indices = LODModel.IndexBuffer.GetArrayView();
+				I0 = Indices[TriangleIndex * 3 + 0];
+				I1 = Indices[TriangleIndex * 3 + 1];
+				I2 = Indices[TriangleIndex * 3 + 2];
+
+				V0 = LODModel.VertexBuffers.PositionVertexBuffer.VertexPosition(I0);
+				V1 = LODModel.VertexBuffers.PositionVertexBuffer.VertexPosition(I1);
+				V2 = LODModel.VertexBuffers.PositionVertexBuffer.VertexPosition(I2);
+			}
 
 			if (bMeshWasPlane)
 			{
@@ -748,7 +790,7 @@ void FMeshUtilities::GenerateSignedDistanceFieldVolumeData(
 				VolumeDimensions.X,
 				VolumeDimensions.Y,
 				VolumeDimensions.Z,
-				Indices.Num() / 3,
+				NumIndices / 3,
 				MinVolumeDistance,
 				MaxVolumeDistance,
 				*MeshName);
@@ -768,6 +810,7 @@ void FMeshUtilities::GenerateSignedDistanceFieldVolumeData(
 
 void FMeshUtilities::GenerateSignedDistanceFieldVolumeData(
 	FString MeshName,
+	const FSourceMeshDataForDerivedDataTask& SourceMeshData,
 	const FStaticMeshLODResources& LODModel,
 	class FQueuedThreadPool& ThreadPool,
 	const TArray<EBlendMode>& MaterialBlendModes,
diff --git a/Engine/Source/Developer/MeshUtilities/Private/MeshUtilitiesPrivate.h b/Engine/Source/Developer/MeshUtilities/Private/MeshUtilitiesPrivate.h
--- a/Engine/Source/Developer/MeshUtilities/Private/MeshUtilitiesPrivate.h
+++ b/Engine/Source/Developer/MeshUtilities/Private/MeshUtilitiesPrivate.h
@@ -96,6 +96,7 @@ private:
 
 	virtual void GenerateSignedDistanceFieldVolumeData(
 		FString MeshName,
+		const FSourceMeshDataForDerivedDataTask& SourceMeshData,
 		const FStaticMeshLODResources& LODModel,
 		class FQueuedThreadPool& ThreadPool,
 		const TArray<EBlendMode>& MaterialBlendModes,
@@ -108,6 +109,7 @@ private:
 
 	virtual bool GenerateCardRepresentationData(
 		FString MeshName,
+		const FSourceMeshDataForDerivedDataTask& SourceMeshData,
 		const FStaticMeshLODResources& LODModel,
 		class FQueuedThreadPool& ThreadPool,
 		const FBoxSphereBounds& Bounds,
diff --git a/Engine/Source/Developer/MeshUtilities/Public/MeshUtilities.h b/Engine/Source/Developer/MeshUtilities/Public/MeshUtilities.h
--- a/Engine/Source/Developer/MeshUtilities/Public/MeshUtilities.h
+++ b/Engine/Source/Developer/MeshUtilities/Public/MeshUtilities.h
@@ -19,6 +19,7 @@ class UStaticMeshComponent;
 struct FFlattenMaterial;
 struct FRawMesh;
 struct FStaticMeshLODResources;
+class FSourceMeshDataForDerivedDataTask;
 
 typedef FIntPoint FMeshIdAndLOD;
 struct FFlattenMaterial;
@@ -205,6 +206,7 @@ public:
 	/** Builds a signed distance field volume for the given LODModel. */
 	virtual void GenerateSignedDistanceFieldVolumeData(
 		FString MeshName,
+		const FSourceMeshDataForDerivedDataTask& SourceMeshData,
 		const FStaticMeshLODResources& LODModel,
 		class FQueuedThreadPool& ThreadPool,
 		const TArray<EBlendMode>& MaterialBlendModes,
@@ -215,6 +217,7 @@ public:
 
 	virtual bool GenerateCardRepresentationData(
 		FString MeshName,
+		const FSourceMeshDataForDerivedDataTask& SourceMeshData,
 		const FStaticMeshLODResources& LODModel,
 		class FQueuedThreadPool& ThreadPool,
 		const FBoxSphereBounds& Bounds,
diff --git a/Engine/Source/Runtime/Engine/Private/DistanceFieldAtlas.cpp b/Engine/Source/Runtime/Engine/Private/DistanceFieldAtlas.cpp
--- a/Engine/Source/Runtime/Engine/Private/DistanceFieldAtlas.cpp
+++ b/Engine/Source/Runtime/Engine/Private/DistanceFieldAtlas.cpp
@@ -21,6 +21,7 @@
 #include "GlobalShader.h"
 #include "RenderGraph.h"
 #include "MeshCardRepresentation.h"
+#include "IMeshBuilderModule.h"
 
 #if WITH_EDITOR
 #include "DerivedDataCacheInterface.h"
@@ -871,7 +872,7 @@ FDistanceFieldAsyncQueue* GDistanceFieldAsyncQueue = NULL;
 #if WITH_EDITOR
 
 // DDC key for distance field data, must be changed when modifying the generation code or data format
-#define DISTANCEFIELD_DERIVEDDATA_VER TEXT("6CBBF5D788CA4699B140BAEC2A3B6B67")
+#define DISTANCEFIELD_DERIVEDDATA_VER TEXT("CA7B773FC3D44F4B8E72678A8BE829DA")
 
 FString BuildDistanceFieldDerivedDataKey(const FString& InMeshKey)
 {
@@ -905,7 +906,7 @@ FString BuildDistanceFieldDerivedDataKey(const FString& InMeshKey)
 
 #if WITH_EDITORONLY_DATA
 
-void FDistanceFieldVolumeData::CacheDerivedData(const FString& InDDCKey, UStaticMesh* Mesh, FStaticMeshRenderData& RenderData, UStaticMesh* GenerateSource, float DistanceFieldResolutionScale, bool bGenerateDistanceFieldAsIfTwoSided)
+void FDistanceFieldVolumeData::CacheDerivedData(const FString& InDDCKey, const ITargetPlatform* TargetPlatform, UStaticMesh* Mesh, FStaticMeshRenderData& RenderData, UStaticMesh* GenerateSource, float DistanceFieldResolutionScale, bool bGenerateDistanceFieldAsIfTwoSided)
 {
 	TArray<uint8> DerivedData;
 
@@ -916,7 +917,7 @@ void FDistanceFieldVolumeData::CacheDerivedData(const FString& InDDCKey, UStatic
 		FMemoryReader Ar(DerivedData, /*bIsPersistent=*/ true);
 		Ar << *this;
 
-		BeginCacheMeshCardRepresentation(Mesh, RenderData, InDDCKey);
+		BeginCacheMeshCardRepresentation(TargetPlatform, Mesh, RenderData, InDDCKey, /*OptionalSourceMeshData*/ nullptr);
 	}
 	else
 	{
@@ -925,6 +926,7 @@ void FDistanceFieldVolumeData::CacheDerivedData(const FString& InDDCKey, UStatic
 		FAsyncDistanceFieldTask* NewTask = new FAsyncDistanceFieldTask;
 		NewTask->DDCKey = InDDCKey;
 		check(Mesh && GenerateSource);
+		NewTask->TargetPlatform = TargetPlatform;
 		NewTask->StaticMesh = Mesh;
 		NewTask->GenerateSource = GenerateSource;
 		NewTask->DistanceFieldResolutionScale = DistanceFieldResolutionScale;
@@ -945,6 +947,22 @@ void FDistanceFieldVolumeData::CacheDerivedData(const FString& InDDCKey, UStatic
 			NewTask->MaterialBlendModes.Add(BlendMode);
 		}
 
+		// Nanite overrides source static mesh with a coarse representation. Need to load original data before we build the mesh SDF.
+		if (Mesh->NaniteSettings.bEnabled)
+		{
+			TRACE_CPUPROFILER_EVENT_SCOPE_TEXT(TEXT("FDistanceFieldVolumeData::BuildMesh"));
+
+			uint32 NumTexCoords = 0;
+			bool bHasColors = false;
+
+			IMeshBuilderModule& MeshBuilderModule = IMeshBuilderModule::GetForPlatform(TargetPlatform);
+			if (!MeshBuilderModule.BuildMesh(Mesh, NewTask->SourceMeshData.Vertices, NewTask->SourceMeshData.Indices, NewTask->SourceMeshData.Sections, /* bBuildOnlyPosition */ true, NumTexCoords, bHasColors))
+			{
+				UE_LOG(LogStaticMesh, Error, TEXT("Failed to build static mesh. See previous line(s) for details."));
+				return;
+			}
+		}
+
 		GDistanceFieldAsyncQueue->AddTask(NewTask);
 	}
 }
@@ -1212,6 +1230,7 @@ void FDistanceFieldAsyncQueue::Build(FAsyncDistanceFieldTask* Task, FQueuedThrea
 
 		MeshUtilities->GenerateSignedDistanceFieldVolumeData(
 			Task->StaticMesh->GetName(),
+			Task->SourceMeshData,
 			LODModel,
 			ThreadPool,
 			Task->MaterialBlendModes,
@@ -1292,7 +1311,7 @@ void FDistanceFieldAsyncQueue::ProcessAsyncTasks()
 				COOK_STAT(Timer.AddMiss(DerivedData.Num()));
 			}
 
-			BeginCacheMeshCardRepresentation(Task->StaticMesh, *Task->StaticMesh->RenderData, Task->DDCKey);
+			BeginCacheMeshCardRepresentation(Task->TargetPlatform, Task->StaticMesh, *Task->StaticMesh->RenderData, Task->DDCKey, &Task->SourceMeshData);
 		}
 
 		delete Task;
diff --git a/Engine/Source/Runtime/Engine/Private/MeshCardRepresentation.cpp b/Engine/Source/Runtime/Engine/Private/MeshCardRepresentation.cpp
--- a/Engine/Source/Runtime/Engine/Private/MeshCardRepresentation.cpp
+++ b/Engine/Source/Runtime/Engine/Private/MeshCardRepresentation.cpp
@@ -17,6 +17,7 @@
 #include "Engine/StaticMesh.h"
 #include "Misc/AutomationTest.h"
 #include "Async/ParallelFor.h"
+#include "IMeshBuilderModule.h"
 
 #if WITH_EDITOR
 #include "DerivedDataCacheInterface.h"
@@ -51,7 +52,7 @@ FCardRepresentationAsyncQueue* GCardRepresentationAsyncQueue = NULL;
 #if WITH_EDITOR
 
 // DDC key for card representation data, must be changed when modifying the generation code or data format
-#define CARDREPRESENTATION_DERIVEDDATA_VER TEXT("BD2E2E564F0C5A69A37178216589A7D2")
+#define CARDREPRESENTATION_DERIVEDDATA_VER TEXT("378A453D4B7A4B163E62A302B1EE8BD8")
 
 FString BuildCardRepresentationDerivedDataKey(const FString& InMeshKey)
 {
@@ -67,7 +68,7 @@ FString BuildCardRepresentationDerivedDataKey(const FString& InMeshKey)
 
 #if WITH_EDITORONLY_DATA
 
-void BeginCacheMeshCardRepresentation(UStaticMesh* StaticMeshAsset, FStaticMeshRenderData& RenderData, const FString& DistanceFieldKey)
+void BeginCacheMeshCardRepresentation(const ITargetPlatform* TargetPlatform, UStaticMesh* StaticMeshAsset, FStaticMeshRenderData& RenderData, const FString& DistanceFieldKey, FSourceMeshDataForDerivedDataTask* OptionalSourceMeshData)
 {
 	static const auto CVarCards = IConsoleManager::Get().FindTConsoleVariableDataInt(TEXT("r.MeshCardRepresentation"));
 
@@ -83,12 +84,12 @@ void BeginCacheMeshCardRepresentation(UStaticMesh* StaticMeshAsset, FStaticMeshR
 
 			UStaticMesh* MeshToGenerateFrom = StaticMeshAsset;
 
-			RenderData.LODResources[0].CardRepresentationData->CacheDerivedData(Key, StaticMeshAsset, MeshToGenerateFrom);
+			RenderData.LODResources[0].CardRepresentationData->CacheDerivedData(Key, TargetPlatform, StaticMeshAsset, MeshToGenerateFrom, OptionalSourceMeshData);
 		}
 	}
 }
 
-void FCardRepresentationData::CacheDerivedData(const FString& InDDCKey, UStaticMesh* Mesh, UStaticMesh* GenerateSource)
+void FCardRepresentationData::CacheDerivedData(const FString& InDDCKey, const ITargetPlatform* TargetPlatform, UStaticMesh* Mesh, UStaticMesh* GenerateSource, FSourceMeshDataForDerivedDataTask* OptionalSourceMeshData)
 {
 	TArray<uint8> DerivedData;
 
@@ -110,6 +111,26 @@ void FCardRepresentationData::CacheDerivedData(const FString& InDDCKey, UStaticM
 		NewTask->GenerateSource = GenerateSource;
 		NewTask->GeneratedCardRepresentation = new FCardRepresentationData();
 
+		// Nanite overrides source static mesh with a coarse representation. Need to load original data before we build the mesh SDF.
+		if (OptionalSourceMeshData)
+		{
+			NewTask->SourceMeshData = *OptionalSourceMeshData;
+		}
+		else if (Mesh->NaniteSettings.bEnabled)
+		{
+			TRACE_CPUPROFILER_EVENT_SCOPE_TEXT(TEXT("FCardRepresentationData::BuildMesh"));
+
+			uint32 NumTexCoords = 0;
+			bool bHasColors = false;
+
+			IMeshBuilderModule& MeshBuilderModule = IMeshBuilderModule::GetForPlatform(TargetPlatform);
+			if (!MeshBuilderModule.BuildMesh(Mesh, NewTask->SourceMeshData.Vertices, NewTask->SourceMeshData.Indices, NewTask->SourceMeshData.Sections, /* bBuildOnlyPosition */ true, NumTexCoords, bHasColors))
+			{
+				UE_LOG(LogStaticMesh, Error, TEXT("Failed to build static mesh. See previous line(s) for details."));
+				return;
+			}
+		}
+
 		GCardRepresentationAsyncQueue->AddTask(NewTask);
 	}
 }
@@ -372,6 +393,7 @@ void FCardRepresentationAsyncQueue::Build(FAsyncCardRepresentationTask* Task, FQ
 
 		Task->bSuccess = MeshUtilities->GenerateCardRepresentationData(
 			Task->StaticMesh->GetName(),
+			Task->SourceMeshData,
 			LODModel,
 			ThreadPool,
 			Task->GenerateSource->RenderData->Bounds,
diff --git a/Engine/Source/Runtime/Engine/Private/StaticMesh.cpp b/Engine/Source/Runtime/Engine/Private/StaticMesh.cpp
--- a/Engine/Source/Runtime/Engine/Private/StaticMesh.cpp
+++ b/Engine/Source/Runtime/Engine/Private/StaticMesh.cpp
@@ -2658,7 +2658,7 @@ void FStaticMeshRenderData::Cache(const ITargetPlatform* TargetPlatform, UStatic
 				BuildSettings.DistanceFieldReplacementMesh->ConditionalPostLoad();
 			}
 
-			LODResources[0].DistanceFieldData->CacheDerivedData(DistanceFieldKey, Owner, *this, MeshToGenerateFrom, BuildSettings.DistanceFieldResolutionScale, BuildSettings.bGenerateDistanceFieldAsIfTwoSided);
+			LODResources[0].DistanceFieldData->CacheDerivedData(DistanceFieldKey, TargetPlatform, Owner, *this, MeshToGenerateFrom, BuildSettings.DistanceFieldResolutionScale, BuildSettings.bGenerateDistanceFieldAsIfTwoSided);
 		}
 		else
 		{
@@ -2667,7 +2667,7 @@ void FStaticMeshRenderData::Cache(const ITargetPlatform* TargetPlatform, UStatic
 	}
 	else
 	{
-		BeginCacheMeshCardRepresentation(Owner, *this, DerivedDataKey);
+		BeginCacheMeshCardRepresentation(TargetPlatform, Owner, *this, DerivedDataKey, /* OptionalSourceMeshData */ nullptr);
 	}
 }
 #endif // #if WITH_EDITOR
diff --git a/Engine/Source/Runtime/Engine/Public/DerivedMeshDataTaskUtils.h b/Engine/Source/Runtime/Engine/Public/DerivedMeshDataTaskUtils.h
new file mode 100644
--- /dev/null
+++ b/Engine/Source/Runtime/Engine/Public/DerivedMeshDataTaskUtils.h
@@ -0,0 +1,21 @@
+// Copyright Epic Games, Inc. All Rights Reserved.
+
+/*=============================================================================
+	SourceMeshDataForDerivedDataTask.h
+=============================================================================*/
+
+#pragma once
+
+/** Source mesh data. Valid only in specific cases, when StaticMesh doesn't contain original data anymore (e.g. it's replaced by the Nanite coarse mesh representation) */
+class FSourceMeshDataForDerivedDataTask
+{
+public:
+	TArray<FStaticMeshBuildVertex> Vertices;
+	TArray<uint32> Indices;
+	FStaticMeshSectionArray Sections;
+
+	bool IsValid() const
+	{ 
+		return Sections.Num() > 0;
+	}
+};
\ No newline at end of file
diff --git a/Engine/Source/Runtime/Engine/Public/DistanceFieldAtlas.h b/Engine/Source/Runtime/Engine/Public/DistanceFieldAtlas.h
--- a/Engine/Source/Runtime/Engine/Public/DistanceFieldAtlas.h
+++ b/Engine/Source/Runtime/Engine/Public/DistanceFieldAtlas.h
@@ -15,6 +15,7 @@
 #include "RenderingThread.h"
 #include "TextureLayout3d.h"
 #include "Templates/UniquePtr.h"
+#include "DerivedMeshDataTaskUtils.h"
 
 class FDistanceFieldVolumeData;
 class UStaticMesh;
@@ -363,7 +364,7 @@ public:
 
 #if WITH_EDITORONLY_DATA
 
-	void CacheDerivedData(const FString& InDDCKey, UStaticMesh* Mesh, class FStaticMeshRenderData& RenderData, UStaticMesh* GenerateSource, float DistanceFieldResolutionScale, bool bGenerateDistanceFieldAsIfTwoSided);
+	void CacheDerivedData(const FString& InDDCKey, const ITargetPlatform* TargetPlatform, UStaticMesh* Mesh, class FStaticMeshRenderData& RenderData, UStaticMesh* GenerateSource, float DistanceFieldResolutionScale, bool bGenerateDistanceFieldAsIfTwoSided);
 
 #endif
 
@@ -382,10 +383,12 @@ public:
 	FAsyncDistanceFieldTask();
 
 	TArray<EBlendMode> MaterialBlendModes;
+	FSourceMeshDataForDerivedDataTask SourceMeshData;
 	UStaticMesh* StaticMesh;
 	UStaticMesh* GenerateSource;
 	float DistanceFieldResolutionScale;
 	bool bGenerateDistanceFieldAsIfTwoSided;
+	const ITargetPlatform* TargetPlatform;
 	FString DDCKey;
 	FDistanceFieldVolumeData* GeneratedVolumeData;
 };
diff --git a/Engine/Source/Runtime/Engine/Public/MeshCardRepresentation.h b/Engine/Source/Runtime/Engine/Public/MeshCardRepresentation.h
--- a/Engine/Source/Runtime/Engine/Public/MeshCardRepresentation.h
+++ b/Engine/Source/Runtime/Engine/Public/MeshCardRepresentation.h
@@ -14,6 +14,7 @@
 #include "RenderResource.h"
 #include "RenderingThread.h"
 #include "Templates/UniquePtr.h"
+#include "DerivedMeshDataTaskUtils.h"
 
 template <class T> class TLockFreePointerListLIFO;
 
@@ -142,7 +143,7 @@ public:
 
 #if WITH_EDITORONLY_DATA
 
-	void CacheDerivedData(const FString& InDDCKey, UStaticMesh* Mesh, UStaticMesh* GenerateSource);
+	void CacheDerivedData(const FString& InDDCKey, const ITargetPlatform* TargetPlatform, UStaticMesh* Mesh, UStaticMesh* GenerateSource, FSourceMeshDataForDerivedDataTask* OptionalSourceMeshData);
 
 #endif
 
@@ -160,6 +161,8 @@ public:
 	FAsyncCardRepresentationTask();
 
 	bool bSuccess = false;
+
+	FSourceMeshDataForDerivedDataTask SourceMeshData;
 	UStaticMesh* StaticMesh;
 	UStaticMesh* GenerateSource;
 	FString DDCKey;
@@ -228,4 +231,4 @@ extern ENGINE_API FCardRepresentationAsyncQueue* GCardRepresentationAsyncQueue;
 
 extern ENGINE_API FString BuildCardRepresentationDerivedDataKey(const FString& InMeshKey);
 
-extern ENGINE_API void BeginCacheMeshCardRepresentation(UStaticMesh* StaticMeshAsset, class FStaticMeshRenderData& RenderData, const FString& DistanceFieldKey);
\ No newline at end of file
+extern ENGINE_API void BeginCacheMeshCardRepresentation(const ITargetPlatform* TargetPlatform, UStaticMesh* StaticMeshAsset, class FStaticMeshRenderData& RenderData, const FString& DistanceFieldKey, FSourceMeshDataForDerivedDataTask* OptionalSourceMeshData);
\ No newline at end of file
