Added FMeshBuilder::BuildMeshVertexPositions. It replaces FMeshBuilder::BuildMesh(bool bBuildOnlyPosition = true) and does the minimal work which is required to extract vertex positions from a mesh descriptor.

This is later used to build static mesh derived data in cases where Nanite coarse mesh representation replaces original static mesh data and we can't depend on it. In extreme cases it lowers main thread derived data build time from ~7.5s to ~0.1s.

diff --git a/Engine/Source/Developer/MeshBuilder/Private/MeshBuilderModule.cpp b/Engine/Source/Developer/MeshBuilder/Private/MeshBuilderModule.cpp
--- a/Engine/Source/Developer/MeshBuilder/Private/MeshBuilderModule.cpp
+++ b/Engine/Source/Developer/MeshBuilder/Private/MeshBuilderModule.cpp
@@ -33,10 +33,14 @@ public:
 		TArray< FStaticMeshBuildVertex >& Verts,
 		TArray< uint32 >& Indexes,
 		FStaticMeshSectionArray& Sections,
-		bool bBuildOnlyPosition,
 		uint32& NumTexCoords,
 		bool& bHasColors ) override;
 
+	virtual bool BuildMeshVertexPositions(
+		UObject* StaticMesh,
+		TArray<uint32>& Indices,
+		TArray<FVector>& Vertices) override;
+
 	virtual bool BuildSkeletalMesh(USkeletalMesh* SkeletalMesh, const int32 LODIndex, const bool bRegenDepLODs) override;
 
 private:
@@ -61,15 +65,28 @@ bool FMeshBuilderModule::BuildMesh(
 	TArray< FStaticMeshBuildVertex >& Verts,
 	TArray< uint32 >& Indexes,
 	FStaticMeshSectionArray& Sections,
-	bool bBuildOnlyPosition,
 	uint32& NumTexCoords,
-	bool& bHasColors )
+	bool& bHasColors)
+{
+	UStaticMesh* StaticMesh = Cast< UStaticMesh >(Mesh);
+	if (StaticMesh)
+	{
+		//Call the static mesh builder
+		return FStaticMeshBuilder().Build(StaticMesh, Verts, Indexes, Sections, NumTexCoords, bHasColors);
+	}
+	return false;
+}
+
+bool FMeshBuilderModule::BuildMeshVertexPositions(
+	UObject* Mesh,
+	TArray<uint32>& Indices,
+	TArray<FVector>& Vertices)
 {
-	UStaticMesh* StaticMesh = Cast< UStaticMesh >( Mesh );
-	if( StaticMesh )
+	UStaticMesh* StaticMesh = Cast< UStaticMesh >(Mesh);
+	if (StaticMesh)
 	{
 		//Call the static mesh builder
-		return FStaticMeshBuilder().Build( StaticMesh, Verts, Indexes, Sections, bBuildOnlyPosition, NumTexCoords, bHasColors );
+		return FStaticMeshBuilder().BuildMeshVertexPositions(StaticMesh, Indices, Vertices);
 	}
 	return false;
 }
diff --git a/Engine/Source/Developer/MeshBuilder/Private/MeshDescriptionHelper.cpp b/Engine/Source/Developer/MeshBuilder/Private/MeshDescriptionHelper.cpp
--- a/Engine/Source/Developer/MeshBuilder/Private/MeshDescriptionHelper.cpp
+++ b/Engine/Source/Developer/MeshBuilder/Private/MeshDescriptionHelper.cpp
@@ -30,7 +30,7 @@ FMeshDescriptionHelper::FMeshDescriptionHelper(FMeshBuildSettings* InBuildSettin
 {
 }
 
-void FMeshDescriptionHelper::SetupRenderMeshDescription(UObject* Owner, FMeshDescription& RenderMeshDescription, bool bBuildOnlyPosition)
+void FMeshDescriptionHelper::SetupRenderMeshDescription(UObject* Owner, FMeshDescription& RenderMeshDescription)
 {
 	TRACE_CPUPROFILER_EVENT_SCOPE(FMeshDescriptionHelper::GetRenderMeshDescription);
 
@@ -53,21 +53,18 @@ void FMeshDescriptionHelper::SetupRenderMeshDescription(UObject* Owner, FMeshDes
 	// Find overlapping corners to accelerate adjacency.
 	FStaticMeshOperations::FindOverlappingCorners(OverlappingCorners, RenderMeshDescription, ComparisonThreshold);
 
+	// Static meshes always blend normals of overlapping corners.
+	EComputeNTBsFlags ComputeNTBsOptions = EComputeNTBsFlags::BlendOverlappingNormals;
+	ComputeNTBsOptions |= BuildSettings->bRemoveDegenerates ? EComputeNTBsFlags::IgnoreDegenerateTriangles : EComputeNTBsFlags::None;
+	ComputeNTBsOptions |= BuildSettings->bComputeWeightedNormals ? EComputeNTBsFlags::WeightedNTBs : EComputeNTBsFlags::None;
+	ComputeNTBsOptions |= BuildSettings->bRecomputeNormals ? EComputeNTBsFlags::Normals : EComputeNTBsFlags::None;
+	ComputeNTBsOptions |= BuildSettings->bRecomputeTangents && !StaticMesh->NaniteSettings.bEnabled ? EComputeNTBsFlags::Tangents : EComputeNTBsFlags::None;	// Nanite doesn't need tangents
+	ComputeNTBsOptions |= BuildSettings->bUseMikkTSpace ? EComputeNTBsFlags::UseMikkTSpace : EComputeNTBsFlags::None;
+
 	// Compute any missing normals or tangents.
-	if (!bBuildOnlyPosition)
-	{
-		// Static meshes always blend normals of overlapping corners.
-		EComputeNTBsFlags ComputeNTBsOptions = EComputeNTBsFlags::BlendOverlappingNormals;
-		ComputeNTBsOptions |= BuildSettings->bRemoveDegenerates ? EComputeNTBsFlags::IgnoreDegenerateTriangles : EComputeNTBsFlags::None;
-		ComputeNTBsOptions |= BuildSettings->bComputeWeightedNormals ? EComputeNTBsFlags::WeightedNTBs : EComputeNTBsFlags::None;
-		ComputeNTBsOptions |= BuildSettings->bRecomputeNormals ? EComputeNTBsFlags::Normals : EComputeNTBsFlags::None;
-		ComputeNTBsOptions |= BuildSettings->bRecomputeTangents && !StaticMesh->NaniteSettings.bEnabled ? EComputeNTBsFlags::Tangents : EComputeNTBsFlags::None;	// Nanite doesn't need tangents
-		ComputeNTBsOptions |= BuildSettings->bUseMikkTSpace ? EComputeNTBsFlags::UseMikkTSpace : EComputeNTBsFlags::None;
-
-		FStaticMeshOperations::ComputeTangentsAndNormals(RenderMeshDescription, ComputeNTBsOptions);
-	}
+	FStaticMeshOperations::ComputeTangentsAndNormals(RenderMeshDescription, ComputeNTBsOptions);
 
-	if (BuildSettings->bGenerateLightmapUVs && !bBuildOnlyPosition && VertexInstanceArray.Num() > 0)
+	if (BuildSettings->bGenerateLightmapUVs && VertexInstanceArray.Num() > 0)
 	{
 		TVertexInstanceAttributesRef<FVector2D> VertexInstanceUVs = Attributes.GetVertexInstanceUVs();
 		int32 NumIndices = VertexInstanceUVs.GetNumChannels();
diff --git a/Engine/Source/Developer/MeshBuilder/Private/MeshDescriptionHelper.h b/Engine/Source/Developer/MeshBuilder/Private/MeshDescriptionHelper.h
--- a/Engine/Source/Developer/MeshBuilder/Private/MeshDescriptionHelper.h
+++ b/Engine/Source/Developer/MeshBuilder/Private/MeshDescriptionHelper.h
@@ -21,7 +21,7 @@ public:
 	FMeshDescriptionHelper(FMeshBuildSettings* InBuildSettings);
 
 	//Build a render mesh description with the BuildSettings. This will update the RenderMeshDescription in place
-	void SetupRenderMeshDescription(UObject* Owner, FMeshDescription& RenderMeshDescription, bool bBuildOnlyPosition = false);
+	void SetupRenderMeshDescription(UObject* Owner, FMeshDescription& RenderMeshDescription);
 
 	void ReduceLOD(const FMeshDescription& BaseMesh, FMeshDescription& DestMesh, const struct FMeshReductionSettings& ReductionSettings, const FOverlappingCorners& InOverlappingCorners, float &OutMaxDeviation);
 
diff --git a/Engine/Source/Developer/MeshBuilder/Private/StaticMeshBuilder.cpp b/Engine/Source/Developer/MeshBuilder/Private/StaticMeshBuilder.cpp
--- a/Engine/Source/Developer/MeshBuilder/Private/StaticMeshBuilder.cpp
+++ b/Engine/Source/Developer/MeshBuilder/Private/StaticMeshBuilder.cpp
@@ -378,7 +378,6 @@ bool FStaticMeshBuilder::Build(
 	TArray< FStaticMeshBuildVertex >& Verts,
 	TArray< uint32 >& Indexes,
 	FStaticMeshSectionArray& Sections,
-	bool bBuildOnlyPosition,
 	uint32& NumTexCoords,
 	bool& bHasColors )
 {
@@ -402,7 +401,7 @@ bool FStaticMeshBuilder::Build(
 	{
 		TRACE_CPUPROFILER_EVENT_SCOPE_TEXT(TEXT("FStaticMeshBuilder::Build::Mappings"));
 		MeshDescription = *OriginalMeshDescription;
-		MeshDescriptionHelper.SetupRenderMeshDescription(StaticMesh, MeshDescription, bBuildOnlyPosition);
+		MeshDescriptionHelper.SetupRenderMeshDescription(StaticMesh, MeshDescription);
 	}
 	else
 	{
@@ -465,6 +464,57 @@ bool FStaticMeshBuilder::Build(
 	return true;
 }
 
+bool FStaticMeshBuilder::BuildMeshVertexPositions(
+	UStaticMesh* StaticMesh,
+	TArray<uint32>& BuiltIndices,
+	TArray<FVector>& BuiltVertices)
+{
+	TRACE_CPUPROFILER_EVENT_SCOPE_TEXT(TEXT("FStaticMeshBuilder::BuildMeshVertexPositions"));
+
+	const FMeshDescription* MeshDescription = StaticMesh->GetMeshDescription(0);
+	if (MeshDescription == nullptr)
+	{
+		//Warn the user that there is no mesh description data
+		UE_LOG(LogStaticMeshBuilder, Error, TEXT("Cannot find a valid mesh description to build the asset."));
+		return false;
+	}
+
+	const FMeshBuildSettings& BuildSettings = StaticMesh->GetSourceModel(0).BuildSettings;
+
+	const FStaticMeshConstAttributes Attributes(*MeshDescription);
+	TArrayView<const FVector> VertexPositions = Attributes.GetVertexPositions().GetRawArray();
+	TArrayView<const FVertexID> VertexIndices = Attributes.GetTriangleVertexIndices().GetRawArray();
+
+	BuiltVertices.Reserve(VertexPositions.Num());
+	for (int32 VertexIndex = 0; VertexIndex < VertexPositions.Num(); ++VertexIndex)
+	{
+		BuiltVertices.Add(VertexPositions[VertexIndex] * BuildSettings.BuildScale3D);
+	}
+
+	BuiltIndices.Reserve(VertexIndices.Num());
+	for (int32 TriangleIndex = 0; TriangleIndex < VertexIndices.Num() / 3; ++TriangleIndex)
+	{
+		const uint32 I0 = VertexIndices[TriangleIndex * 3 + 0];
+		const uint32 I1 = VertexIndices[TriangleIndex * 3 + 1];
+		const uint32 I2 = VertexIndices[TriangleIndex * 3 + 2];
+
+		const FVector V0 = BuiltVertices[I0];
+		const FVector V1 = BuiltVertices[I1];
+		const FVector V2 = BuiltVertices[I2];
+
+		const FVector TriangleNormal = ((V1 - V2) ^ (V0 - V2));
+		const bool bDegenerateTriangle = TriangleNormal.SizeSquared() < SMALL_NUMBER;
+		if (!bDegenerateTriangle)
+		{
+			BuiltIndices.Add(I0);
+			BuiltIndices.Add(I1);
+			BuiltIndices.Add(I2);
+		}
+	}
+
+	return true;
+}
+
 bool AreVerticesEqual(FStaticMeshBuildVertex const& A, FStaticMeshBuildVertex const& B, float ComparisonThreshold)
 {
 	if (   !A.Position.Equals(B.Position, ComparisonThreshold)
diff --git a/Engine/Source/Developer/MeshBuilder/Public/IMeshBuilderModule.h b/Engine/Source/Developer/MeshBuilder/Public/IMeshBuilderModule.h
--- a/Engine/Source/Developer/MeshBuilder/Public/IMeshBuilderModule.h
+++ b/Engine/Source/Developer/MeshBuilder/Public/IMeshBuilderModule.h
@@ -36,12 +36,19 @@ public:
 		TArray< FStaticMeshBuildVertex >& Verts,
 		TArray< uint32 >& Indexes,
 		FStaticMeshSectionArray& Sections,
-		bool bBuildOnlyPosition,
 		uint32& NumTexCoords,
 		bool& bHasColors )
 	{
 		return false;
 	}
 
+	virtual bool BuildMeshVertexPositions(
+		class UObject* StaticMesh,
+		TArray<uint32>& Indices,
+		TArray<FVector>& Vertices)
+	{
+		return false;
+	}
+
 	virtual bool BuildSkeletalMesh(class USkeletalMesh* SkeletalMesh, int32 LODIndex, const bool bRegenDepLODs) { return false; }
 };
diff --git a/Engine/Source/Developer/MeshBuilder/Public/MeshBuilder.h b/Engine/Source/Developer/MeshBuilder/Public/MeshBuilder.h
--- a/Engine/Source/Developer/MeshBuilder/Public/MeshBuilder.h
+++ b/Engine/Source/Developer/MeshBuilder/Public/MeshBuilder.h
@@ -31,9 +31,12 @@ public:
 		TArray< FStaticMeshBuildVertex >& Verts,
 		TArray< uint32 >& Indexes,
 		FStaticMeshSectionArray& Sections,
-		bool bBuildOnlyPosition,
 		uint32& NumTexCoords,
 		bool& bHasColors ) = 0;
+	virtual bool BuildMeshVertexPositions(
+		UStaticMesh* StaticMesh,
+		TArray<uint32>& Indices,
+		TArray<FVector>& Vertices) = 0;
 
 	/**
 	 * Build function should be override and is the starting point for skeletal mesh builders
diff --git a/Engine/Source/Developer/MeshBuilder/Public/SkeletalMeshBuilder.h b/Engine/Source/Developer/MeshBuilder/Public/SkeletalMeshBuilder.h
--- a/Engine/Source/Developer/MeshBuilder/Public/SkeletalMeshBuilder.h
+++ b/Engine/Source/Developer/MeshBuilder/Public/SkeletalMeshBuilder.h
@@ -30,7 +30,6 @@ public:
 		TArray< FStaticMeshBuildVertex >& Verts,
 		TArray< uint32 >& Indexes,
 		FStaticMeshSectionArray& Sections,
-		bool bBuildOnlyPosition,
 		uint32& NumTexCoords,
 		bool& bHasColors) override
 	{
@@ -38,6 +37,16 @@ public:
 		check(No_Support_For_StaticMesh_Build_In_FSkeletalMeshBuilder_Class);
 		return false;
 	}
+
+	virtual bool BuildMeshVertexPositions(
+		UStaticMesh* StaticMesh,
+		TArray<uint32>& Indices,
+		TArray<FVector>& Vertices) override
+	{
+		bool No_Support_For_StaticMesh_Build_In_FSkeletalMeshBuilder_Class = false;
+		check(No_Support_For_StaticMesh_Build_In_FSkeletalMeshBuilder_Class);
+		return false;
+	}
 	
 	virtual bool Build(USkeletalMesh* SkeletalMesh, const int32 LODIndex, const bool bRegenDepLODs) override;
 
diff --git a/Engine/Source/Developer/MeshBuilder/Public/StaticMeshBuilder.h b/Engine/Source/Developer/MeshBuilder/Public/StaticMeshBuilder.h
--- a/Engine/Source/Developer/MeshBuilder/Public/StaticMeshBuilder.h
+++ b/Engine/Source/Developer/MeshBuilder/Public/StaticMeshBuilder.h
@@ -33,10 +33,14 @@ public:
 		TArray< FStaticMeshBuildVertex >& Verts,
 		TArray< uint32 >& Indexes,
 		FStaticMeshSectionArray& Sections,
-		bool bBuildOnlyPosition,
 		uint32& NumTexCoords,
 		bool& bHasColors ) override;
 
+	virtual bool BuildMeshVertexPositions(
+		UStaticMesh* StaticMesh,
+		TArray<uint32>& Indices,
+		TArray<FVector>& Vertices) override;
+
 private:
 
 	void OnBuildRenderMeshStart(class UStaticMesh* StaticMesh, const bool bInvalidateLighting);
diff --git a/Engine/Source/Developer/MeshUtilities/Private/MeshCardRepresentationUtilities.cpp b/Engine/Source/Developer/MeshUtilities/Private/MeshCardRepresentationUtilities.cpp
--- a/Engine/Source/Developer/MeshUtilities/Private/MeshCardRepresentationUtilities.cpp
+++ b/Engine/Source/Developer/MeshUtilities/Private/MeshCardRepresentationUtilities.cpp
@@ -27,7 +27,7 @@ bool SetupEmbreeScene(
 	RTCDevice EmbreeDevice,
 	RTCScene& EmbreeScene)
 {
-	const int32 NumVertices = SourceMeshData.IsValid() ? SourceMeshData.Vertices.Num() : LODModel.VertexBuffers.PositionVertexBuffer.GetNumVertices();
+	const int32 NumVertices = SourceMeshData.IsValid() ? SourceMeshData.GetNumVertices() : LODModel.VertexBuffers.PositionVertexBuffer.GetNumVertices();
 
 	EmbreeScene = rtcDeviceNewScene(EmbreeDevice, RTC_SCENE_STATIC, RTC_INTERSECT1);
 			
@@ -58,13 +58,13 @@ bool SetupEmbreeScene(
 
 		if (SourceMeshData.IsValid())
 		{
-			I0 = SourceMeshData.Indices[TriangleIndex * 3 + 0];
-			I1 = SourceMeshData.Indices[TriangleIndex * 3 + 1];
-			I2 = SourceMeshData.Indices[TriangleIndex * 3 + 2];
+			I0 = SourceMeshData.TriangleIndices[TriangleIndex * 3 + 0];
+			I1 = SourceMeshData.TriangleIndices[TriangleIndex * 3 + 1];
+			I2 = SourceMeshData.TriangleIndices[TriangleIndex * 3 + 2];
 
-			V0 = SourceMeshData.Vertices[I0].Position;
-			V1 = SourceMeshData.Vertices[I1].Position;
-			V2 = SourceMeshData.Vertices[I2].Position;
+			V0 = SourceMeshData.VertexPositions[I0];
+			V1 = SourceMeshData.VertexPositions[I1];
+			V2 = SourceMeshData.VertexPositions[I2];
 		}
 		else
 		{
@@ -553,7 +553,7 @@ bool FMeshUtilities::GenerateCardRepresentationData(
 		return false;
 	}
 
-	const int32 NumIndices = SourceMeshData.IsValid() ? SourceMeshData.Indices.Num() : LODModel.IndexBuffer.GetNumIndices();
+	const int32 NumIndices = SourceMeshData.IsValid() ? SourceMeshData.GetNumIndices() : LODModel.IndexBuffer.GetNumIndices();
 	const int32 NumTriangles = NumIndices / 3;
 
 	FInputPrimitiveIds InputPrimitiveIds;
diff --git a/Engine/Source/Developer/MeshUtilities/Private/MeshDistanceFieldUtilities.cpp b/Engine/Source/Developer/MeshUtilities/Private/MeshDistanceFieldUtilities.cpp
--- a/Engine/Source/Developer/MeshUtilities/Private/MeshDistanceFieldUtilities.cpp
+++ b/Engine/Source/Developer/MeshUtilities/Private/MeshDistanceFieldUtilities.cpp
@@ -336,10 +336,9 @@ void FMeshUtilities::GenerateSignedDistanceFieldVolumeData(
 	{
 		const double StartTime = FPlatformTime::Seconds();
 
-		const int32 NumIndices = SourceMeshData.IsValid() ? SourceMeshData.Indices.Num() : LODModel.IndexBuffer.GetNumIndices();
+		const int32 NumIndices = SourceMeshData.IsValid() ? SourceMeshData.GetNumIndices() : LODModel.IndexBuffer.GetNumIndices();
 		const int32 NumTriangles = NumIndices / 3;
-		const int32 NumVertices = SourceMeshData.IsValid() ? SourceMeshData.Vertices.Num() : LODModel.VertexBuffers.PositionVertexBuffer.GetNumVertices();
-		const FStaticMeshSectionArray& Sections = SourceMeshData.IsValid() ? SourceMeshData.Sections : LODModel.Sections;
+		const int32 NumVertices = SourceMeshData.IsValid() ? SourceMeshData.GetNumVertices() : LODModel.VertexBuffers.PositionVertexBuffer.GetNumVertices();
 
 		TArray<FkDOPBuildCollisionTriangle<uint32> > BuildTriangles;
 
@@ -378,28 +377,38 @@ void FMeshUtilities::GenerateSignedDistanceFieldVolumeData(
 		TArray<int32> FilteredTriangles;
 		FilteredTriangles.Empty(NumTriangles);
 
-		for (int32 TriangleIndex = 0; TriangleIndex < NumTriangles; ++TriangleIndex)
+		if (SourceMeshData.IsValid())
 		{
-			bool bTriangleIsOpaqueOrMasked = false;
-
-			for (int32 SectionIndex = 0; SectionIndex < Sections.Num(); SectionIndex++)
+			for (int32 TriangleIndex = 0; TriangleIndex < NumTriangles; ++TriangleIndex)
 			{
-				const FStaticMeshSection& Section = Sections[SectionIndex];
+				FilteredTriangles.Add(TriangleIndex);
+			}
+		}
+		else
+		{
+			for (int32 TriangleIndex = 0; TriangleIndex < NumTriangles; ++TriangleIndex)
+			{
+				bool bTriangleIsOpaqueOrMasked = false;
 
-				if ((uint32)(TriangleIndex * 3) >= Section.FirstIndex && (uint32)(TriangleIndex * 3) < Section.FirstIndex + Section.NumTriangles * 3)
+				for (int32 SectionIndex = 0; SectionIndex < LODModel.Sections.Num(); SectionIndex++)
 				{
-					if (MaterialBlendModes.IsValidIndex(Section.MaterialIndex))
+					const FStaticMeshSection& Section = LODModel.Sections[SectionIndex];
+
+					if ((uint32)(TriangleIndex * 3) >= Section.FirstIndex && (uint32)(TriangleIndex * 3) < Section.FirstIndex + Section.NumTriangles * 3)
 					{
-						bTriangleIsOpaqueOrMasked = !IsTranslucentBlendMode(MaterialBlendModes[Section.MaterialIndex]);
-					}
+						if (MaterialBlendModes.IsValidIndex(Section.MaterialIndex))
+						{
+							bTriangleIsOpaqueOrMasked = !IsTranslucentBlendMode(MaterialBlendModes[Section.MaterialIndex]);
+						}
 
-					break;
+						break;
+					}
 				}
-			}
 
-			if (bTriangleIsOpaqueOrMasked)
-			{
-				FilteredTriangles.Add(TriangleIndex);
+				if (bTriangleIsOpaqueOrMasked)
+				{
+					FilteredTriangles.Add(TriangleIndex);
+				}
 			}
 		}
 
@@ -433,13 +442,13 @@ void FMeshUtilities::GenerateSignedDistanceFieldVolumeData(
 			const int32 TriangleIndex = FilteredTriangles[FilteredTriangleIndex];
 			if (SourceMeshData.IsValid())
 			{
-				I0 = SourceMeshData.Indices[TriangleIndex * 3 + 0];
-				I1 = SourceMeshData.Indices[TriangleIndex * 3 + 1];
-				I2 = SourceMeshData.Indices[TriangleIndex * 3 + 2];
+				I0 = SourceMeshData.TriangleIndices[TriangleIndex * 3 + 0];
+				I1 = SourceMeshData.TriangleIndices[TriangleIndex * 3 + 1];
+				I2 = SourceMeshData.TriangleIndices[TriangleIndex * 3 + 2];
 
-				V0 = SourceMeshData.Vertices[I0].Position;
-				V1 = SourceMeshData.Vertices[I1].Position;
-				V2 = SourceMeshData.Vertices[I2].Position;
+				V0 = SourceMeshData.VertexPositions[I0];
+				V1 = SourceMeshData.VertexPositions[I1];
+				V2 = SourceMeshData.VertexPositions[I2];
 			}
 			else
 			{
diff --git a/Engine/Source/Runtime/Engine/Private/DistanceFieldAtlas.cpp b/Engine/Source/Runtime/Engine/Private/DistanceFieldAtlas.cpp
--- a/Engine/Source/Runtime/Engine/Private/DistanceFieldAtlas.cpp
+++ b/Engine/Source/Runtime/Engine/Private/DistanceFieldAtlas.cpp
@@ -958,13 +958,8 @@ void FDistanceFieldVolumeData::CacheDerivedData(const FString& InDDCKey, const I
 		// Nanite overrides source static mesh with a coarse representation. Need to load original data before we build the mesh SDF.
 		if (Mesh->NaniteSettings.bEnabled)
 		{
-			TRACE_CPUPROFILER_EVENT_SCOPE_TEXT(TEXT("FDistanceFieldVolumeData::BuildMesh"));
-
-			uint32 NumTexCoords = 0;
-			bool bHasColors = false;
-
 			IMeshBuilderModule& MeshBuilderModule = IMeshBuilderModule::GetForPlatform(TargetPlatform);
-			if (!MeshBuilderModule.BuildMesh(Mesh, NewTask->SourceMeshData.Vertices, NewTask->SourceMeshData.Indices, NewTask->SourceMeshData.Sections, /* bBuildOnlyPosition */ true, NumTexCoords, bHasColors))
+			if (!MeshBuilderModule.BuildMeshVertexPositions(Mesh, NewTask->SourceMeshData.TriangleIndices, NewTask->SourceMeshData.VertexPositions))
 			{
 				UE_LOG(LogStaticMesh, Error, TEXT("Failed to build static mesh. See previous line(s) for details."));
 			}
diff --git a/Engine/Source/Runtime/Engine/Private/MeshCardRepresentation.cpp b/Engine/Source/Runtime/Engine/Private/MeshCardRepresentation.cpp
--- a/Engine/Source/Runtime/Engine/Private/MeshCardRepresentation.cpp
+++ b/Engine/Source/Runtime/Engine/Private/MeshCardRepresentation.cpp
@@ -123,13 +123,8 @@ void FCardRepresentationData::CacheDerivedData(const FString& InDDCKey, const IT
 		}
 		else if (Mesh->NaniteSettings.bEnabled)
 		{
-			TRACE_CPUPROFILER_EVENT_SCOPE_TEXT(TEXT("FCardRepresentationData::BuildMesh"));
-
-			uint32 NumTexCoords = 0;
-			bool bHasColors = false;
-
 			IMeshBuilderModule& MeshBuilderModule = IMeshBuilderModule::GetForPlatform(TargetPlatform);
-			if (!MeshBuilderModule.BuildMesh(Mesh, NewTask->SourceMeshData.Vertices, NewTask->SourceMeshData.Indices, NewTask->SourceMeshData.Sections, /* bBuildOnlyPosition */ true, NumTexCoords, bHasColors))
+			if (!MeshBuilderModule.BuildMeshVertexPositions(Mesh, NewTask->SourceMeshData.TriangleIndices, NewTask->SourceMeshData.VertexPositions))
 			{
 				UE_LOG(LogStaticMesh, Error, TEXT("Failed to build static mesh. See previous line(s) for details."));
 			}
diff --git a/Engine/Source/Runtime/Engine/Public/DerivedMeshDataTaskUtils.h b/Engine/Source/Runtime/Engine/Public/DerivedMeshDataTaskUtils.h
--- a/Engine/Source/Runtime/Engine/Public/DerivedMeshDataTaskUtils.h
+++ b/Engine/Source/Runtime/Engine/Public/DerivedMeshDataTaskUtils.h
@@ -13,12 +13,21 @@
 class FSourceMeshDataForDerivedDataTask
 {
 public:
-	TArray<FStaticMeshBuildVertex> Vertices;
-	TArray<uint32> Indices;
-	FStaticMeshSectionArray Sections;
+	TArray<uint32> TriangleIndices;
+	TArray<FVector> VertexPositions;
+
+	int32 GetNumIndices() const
+	{
+		return TriangleIndices.Num();
+	}
+
+	int32 GetNumVertices() const
+	{
+		return VertexPositions.Num();
+	}
 
 	bool IsValid() const
 	{ 
-		return Sections.Num() > 0;
+		return TriangleIndices.Num() > 0;
 	}
 };
\ No newline at end of file
