SDF and card rep generation now check for degenerate triangles before adding them to the Embree scene

diff --git a/Engine/Source/Developer/MeshUtilities/Private/MeshCardRepresentationUtilities.cpp b/Engine/Source/Developer/MeshUtilities/Private/MeshCardRepresentationUtilities.cpp
--- a/Engine/Source/Developer/MeshUtilities/Private/MeshCardRepresentationUtilities.cpp
+++ b/Engine/Source/Developer/MeshUtilities/Private/MeshCardRepresentationUtilities.cpp
@@ -41,17 +41,63 @@ bool SetupEmbreeScene(
 
 	const int32 NumTriangles = PrimitiveRangeEnd - PrimitiveRangeStart;
 
-	FVector4* EmbreeVertices = NULL;
-	int32* EmbreeIndices = NULL;
+	TArray<int32> FilteredTriangles;
+	FilteredTriangles.Empty(NumTriangles);
 
-	uint32 GeomID = rtcNewTriangleMesh(EmbreeScene, RTC_GEOMETRY_STATIC, NumTriangles, NumVertices);
+	if (SourceMeshData.IsValid())
+	{
+		for (int32 TriangleIndex = 0; TriangleIndex < NumTriangles; ++TriangleIndex)
+		{
+			const uint32 I0 = SourceMeshData.TriangleIndices[TriangleIndex * 3 + 0];
+			const uint32 I1 = SourceMeshData.TriangleIndices[TriangleIndex * 3 + 1];
+			const uint32 I2 = SourceMeshData.TriangleIndices[TriangleIndex * 3 + 2];
+
+			const FVector V0 = SourceMeshData.VertexPositions[I0];
+			const FVector V1 = SourceMeshData.VertexPositions[I1];
+			const FVector V2 = SourceMeshData.VertexPositions[I2];
+
+			const FVector TriangleNormal = ((V1 - V2) ^ (V0 - V2));
+			const bool bDegenerateTriangle = TriangleNormal.SizeSquared() < SMALL_NUMBER;
+			if (!bDegenerateTriangle)
+			{
+				FilteredTriangles.Add(TriangleIndex);
+			}
+		}
+	}
+	else
+	{
+		for (int32 TriangleIndex = 0; TriangleIndex < NumTriangles; ++TriangleIndex)
+		{
+			const FIndexArrayView Indices = LODModel.IndexBuffer.GetArrayView();
+			const uint32 I0 = Indices[TriangleIndex * 3 + 0];
+			const uint32 I1 = Indices[TriangleIndex * 3 + 1];
+			const uint32 I2 = Indices[TriangleIndex * 3 + 2];
+
+			const FVector V0 = LODModel.VertexBuffers.PositionVertexBuffer.VertexPosition(I0);
+			const FVector V1 = LODModel.VertexBuffers.PositionVertexBuffer.VertexPosition(I1);
+			const FVector V2 = LODModel.VertexBuffers.PositionVertexBuffer.VertexPosition(I2);
+
+			const FVector TriangleNormal = ((V1 - V2) ^ (V0 - V2));
+			const bool bDegenerateTriangle = TriangleNormal.SizeSquared() < SMALL_NUMBER;
+			if (!bDegenerateTriangle)
+			{
+				FilteredTriangles.Add(TriangleIndex);
+			}
+		}
+	}
+
+
+	FVector4* EmbreeVertices = nullptr;
+	int32* EmbreeIndices = nullptr;
+
+	uint32 GeomID = rtcNewTriangleMesh(EmbreeScene, RTC_GEOMETRY_STATIC, FilteredTriangles.Num(), NumVertices);
 
 	EmbreeVertices = (FVector4*)rtcMapBuffer(EmbreeScene, GeomID, RTC_VERTEX_BUFFER);
 	EmbreeIndices = (int32*)rtcMapBuffer(EmbreeScene, GeomID, RTC_INDEX_BUFFER);
 
-	for (int32 PrimitiveIndex = PrimitiveRangeStart; PrimitiveIndex < PrimitiveRangeEnd; PrimitiveIndex++)
+	for (int32 FilteredTriangleIndex = 0; FilteredTriangleIndex < FilteredTriangles.Num(); FilteredTriangleIndex++)
 	{
-		int32 TriangleIndex = InputPrimitiveIds[PrimitiveIndex];
+		const int32 TriangleIndex = FilteredTriangles[FilteredTriangleIndex];
 
 		int32 I0, I1, I2;
 		FVector V0, V1, V2;
@@ -78,9 +124,9 @@ bool SetupEmbreeScene(
 			V2 = LODModel.VertexBuffers.PositionVertexBuffer.VertexPosition(I2);
 		}
 
-		EmbreeIndices[TriangleIndex * 3 + 0] = I0;
-		EmbreeIndices[TriangleIndex * 3 + 1] = I1;
-		EmbreeIndices[TriangleIndex * 3 + 2] = I2;
+		EmbreeIndices[FilteredTriangleIndex * 3 + 0] = I0;
+		EmbreeIndices[FilteredTriangleIndex * 3 + 1] = I1;
+		EmbreeIndices[FilteredTriangleIndex * 3 + 2] = I2;
 
 		EmbreeVertices[I0] = FVector4(V0, 0);
 		EmbreeVertices[I1] = FVector4(V1, 0);
diff --git a/Engine/Source/Developer/MeshUtilities/Private/MeshDistanceFieldUtilities.cpp b/Engine/Source/Developer/MeshUtilities/Private/MeshDistanceFieldUtilities.cpp
--- a/Engine/Source/Developer/MeshUtilities/Private/MeshDistanceFieldUtilities.cpp
+++ b/Engine/Source/Developer/MeshUtilities/Private/MeshDistanceFieldUtilities.cpp
@@ -381,33 +381,60 @@ void FMeshUtilities::GenerateSignedDistanceFieldVolumeData(
 		{
 			for (int32 TriangleIndex = 0; TriangleIndex < NumTriangles; ++TriangleIndex)
 			{
-				FilteredTriangles.Add(TriangleIndex);
+				const uint32 I0 = SourceMeshData.TriangleIndices[TriangleIndex * 3 + 0];
+				const uint32 I1 = SourceMeshData.TriangleIndices[TriangleIndex * 3 + 1];
+				const uint32 I2 = SourceMeshData.TriangleIndices[TriangleIndex * 3 + 2];
+
+				const FVector V0 = SourceMeshData.VertexPositions[I0];
+				const FVector V1 = SourceMeshData.VertexPositions[I1];
+				const FVector V2 = SourceMeshData.VertexPositions[I2];
+
+				const FVector TriangleNormal = ((V1 - V2) ^ (V0 - V2));
+				const bool bDegenerateTriangle = TriangleNormal.SizeSquared() < SMALL_NUMBER;
+				if (!bDegenerateTriangle)
+				{
+					FilteredTriangles.Add(TriangleIndex);
+				}
 			}
 		}
 		else
 		{
 			for (int32 TriangleIndex = 0; TriangleIndex < NumTriangles; ++TriangleIndex)
 			{
-				bool bTriangleIsOpaqueOrMasked = false;
+				const FIndexArrayView Indices = LODModel.IndexBuffer.GetArrayView();
+				const uint32 I0 = Indices[TriangleIndex * 3 + 0];
+				const uint32 I1 = Indices[TriangleIndex * 3 + 1];
+				const uint32 I2 = Indices[TriangleIndex * 3 + 2];
+
+				const FVector V0 = LODModel.VertexBuffers.PositionVertexBuffer.VertexPosition(I0);
+				const FVector V1 = LODModel.VertexBuffers.PositionVertexBuffer.VertexPosition(I1);
+				const FVector V2 = LODModel.VertexBuffers.PositionVertexBuffer.VertexPosition(I2);
 
-				for (int32 SectionIndex = 0; SectionIndex < LODModel.Sections.Num(); SectionIndex++)
+				const FVector TriangleNormal = ((V1 - V2) ^ (V0 - V2));
+				const bool bDegenerateTriangle = TriangleNormal.SizeSquared() < SMALL_NUMBER;
+				if (!bDegenerateTriangle)
 				{
-					const FStaticMeshSection& Section = LODModel.Sections[SectionIndex];
+					bool bTriangleIsOpaqueOrMasked = false;
 
-					if ((uint32)(TriangleIndex * 3) >= Section.FirstIndex && (uint32)(TriangleIndex * 3) < Section.FirstIndex + Section.NumTriangles * 3)
+					for (int32 SectionIndex = 0; SectionIndex < LODModel.Sections.Num(); SectionIndex++)
 					{
-						if (MaterialBlendModes.IsValidIndex(Section.MaterialIndex))
+						const FStaticMeshSection& Section = LODModel.Sections[SectionIndex];
+
+						if ((uint32)(TriangleIndex * 3) >= Section.FirstIndex && (uint32)(TriangleIndex * 3) < Section.FirstIndex + Section.NumTriangles * 3)
 						{
-							bTriangleIsOpaqueOrMasked = !IsTranslucentBlendMode(MaterialBlendModes[Section.MaterialIndex]);
-						}
+							if (MaterialBlendModes.IsValidIndex(Section.MaterialIndex))
+							{
+								bTriangleIsOpaqueOrMasked = !IsTranslucentBlendMode(MaterialBlendModes[Section.MaterialIndex]);
+							}
 
-						break;
+							break;
+						}
 					}
-				}
 
-				if (bTriangleIsOpaqueOrMasked)
-				{
-					FilteredTriangles.Add(TriangleIndex);
+					if (bTriangleIsOpaqueOrMasked)
+					{
+						FilteredTriangles.Add(TriangleIndex);
+					}
 				}
 			}
 		}
