Mesh Distance Field supports Two Sided materials properly, when assigned to the asset

diff --git a/Engine/Source/Developer/MeshUtilities/Private/MeshDistanceFieldUtilities.cpp b/Engine/Source/Developer/MeshUtilities/Private/MeshDistanceFieldUtilities.cpp
--- a/Engine/Source/Developer/MeshUtilities/Private/MeshDistanceFieldUtilities.cpp
+++ b/Engine/Source/Developer/MeshUtilities/Private/MeshDistanceFieldUtilities.cpp
@@ -324,7 +324,7 @@ void FMeshUtilities::GenerateSignedDistanceFieldVolumeData(
 	const FSourceMeshDataForDerivedDataTask& SourceMeshData,
 	const FStaticMeshLODResources& LODModel,
 	class FQueuedThreadPool& ThreadPool,
-	const TArray<EBlendMode>& MaterialBlendModes,
+	const TArray<FSignedDistanceFieldBuildMaterialData>& MaterialBlendModes,
 	const FBoxSphereBounds& Bounds,
 	float DistanceFieldResolutionScale,
 	bool bGenerateAsIfTwoSided,
@@ -424,7 +424,7 @@ void FMeshUtilities::GenerateSignedDistanceFieldVolumeData(
 						{
 							if (MaterialBlendModes.IsValidIndex(Section.MaterialIndex))
 							{
-								bTriangleIsOpaqueOrMasked = !IsTranslucentBlendMode(MaterialBlendModes[Section.MaterialIndex]);
+								bTriangleIsOpaqueOrMasked = !IsTranslucentBlendMode(MaterialBlendModes[Section.MaterialIndex].BlendMode);
 							}
 
 							break;
@@ -489,6 +489,23 @@ void FMeshUtilities::GenerateSignedDistanceFieldVolumeData(
 				V2 = LODModel.VertexBuffers.PositionVertexBuffer.VertexPosition(I2);
 			}
 
+			bool bTriangleIsTwoSided = false;
+
+			for (int32 SectionIndex = 0; SectionIndex < LODModel.Sections.Num(); SectionIndex++)
+			{
+				const FStaticMeshSection& Section = LODModel.Sections[SectionIndex];
+
+				if ((uint32)(TriangleIndex * 3) >= Section.FirstIndex && (uint32)(TriangleIndex * 3) < Section.FirstIndex + Section.NumTriangles * 3)
+				{
+					if (MaterialBlendModes.IsValidIndex(Section.MaterialIndex))
+					{
+						bTriangleIsTwoSided = MaterialBlendModes[Section.MaterialIndex].bTwoSided;
+					}
+
+					break;
+				}
+			}
+
 			if (bUseEmbree)
 			{
 				EmbreeIndices[FilteredTriangleIndex * 3 + 0] = I0;
@@ -501,14 +518,14 @@ void FMeshUtilities::GenerateSignedDistanceFieldVolumeData(
 
 				FEmbreeTriangleDesc Desc;
 				// Store bGenerateAsIfTwoSided in material index
-				Desc.ElementIndex = bGenerateAsIfTwoSided ? 1 : 0;
+				Desc.ElementIndex = bGenerateAsIfTwoSided || bTriangleIsTwoSided ? 1 : 0;
 				Geometry.TriangleDescs.Add(Desc);
 			}
 			else
 			{
 				BuildTriangles.Add(FkDOPBuildCollisionTriangle<uint32>(
 					// Store bGenerateAsIfTwoSided in material index
-					bGenerateAsIfTwoSided,
+					bGenerateAsIfTwoSided || bTriangleIsTwoSided ? 1 : 0,
 					V0,
 					V1,
 					V2));
@@ -796,7 +813,7 @@ void FMeshUtilities::GenerateSignedDistanceFieldVolumeData(
 	const FSourceMeshDataForDerivedDataTask& SourceMeshData,
 	const FStaticMeshLODResources& LODModel,
 	class FQueuedThreadPool& ThreadPool,
-	const TArray<EBlendMode>& MaterialBlendModes,
+	const TArray<FSignedDistanceFieldBuildMaterialData>& MaterialBlendModes,
 	const FBoxSphereBounds& Bounds,
 	float DistanceFieldResolutionScale,
 	bool bGenerateAsIfTwoSided,
diff --git a/Engine/Source/Developer/MeshUtilities/Private/MeshUtilitiesPrivate.h b/Engine/Source/Developer/MeshUtilities/Private/MeshUtilitiesPrivate.h
--- a/Engine/Source/Developer/MeshUtilities/Private/MeshUtilitiesPrivate.h
+++ b/Engine/Source/Developer/MeshUtilities/Private/MeshUtilitiesPrivate.h
@@ -99,7 +99,7 @@ private:
 		const FSourceMeshDataForDerivedDataTask& SourceMeshData,
 		const FStaticMeshLODResources& LODModel,
 		class FQueuedThreadPool& ThreadPool,
-		const TArray<EBlendMode>& MaterialBlendModes,
+		const TArray<FSignedDistanceFieldBuildMaterialData>& MaterialBlendModes,
 		const FBoxSphereBounds& Bounds,
 		float DistanceFieldResolutionScale,
 		bool bGenerateAsIfTwoSided,
diff --git a/Engine/Source/Developer/MeshUtilities/Public/MeshUtilities.h b/Engine/Source/Developer/MeshUtilities/Public/MeshUtilities.h
--- a/Engine/Source/Developer/MeshUtilities/Public/MeshUtilities.h
+++ b/Engine/Source/Developer/MeshUtilities/Public/MeshUtilities.h
@@ -50,6 +50,13 @@ struct FBoneVertInfo
 	TArray<FVector>	Normals;
 };
 
+class FSignedDistanceFieldBuildMaterialData
+{
+public:
+	EBlendMode BlendMode;
+	bool bTwoSided;
+};
+
 struct FOverlappingCorners;
 
 class IMeshUtilities : public IModuleInterface
@@ -209,7 +216,7 @@ public:
 		const FSourceMeshDataForDerivedDataTask& SourceMeshData,
 		const FStaticMeshLODResources& LODModel,
 		class FQueuedThreadPool& ThreadPool,
-		const TArray<EBlendMode>& MaterialBlendModes,
+		const TArray<FSignedDistanceFieldBuildMaterialData>& MaterialBlendModes,
 		const FBoxSphereBounds& Bounds,
 		float DistanceFieldResolutionScale,
 		bool bGenerateAsIfTwoSided,
diff --git a/Engine/Source/Runtime/Engine/Private/DistanceFieldAtlas.cpp b/Engine/Source/Runtime/Engine/Private/DistanceFieldAtlas.cpp
--- a/Engine/Source/Runtime/Engine/Private/DistanceFieldAtlas.cpp
+++ b/Engine/Source/Runtime/Engine/Private/DistanceFieldAtlas.cpp
@@ -512,7 +512,7 @@ void FDistanceFieldVolumeTextureAtlas::UpdateAllocations(FRDGBuilder& GraphBuild
 
 				if (!BlockAllocator.AddElement((uint32&)Texture->AtlasAllocationMin.X, (uint32&)Texture->AtlasAllocationMin.Y, (uint32&)Texture->AtlasAllocationMin.Z, Size.X, Size.Y, Size.Z))
 				{
-					UE_LOG(LogStaticMesh, Warning, TEXT("Failed to allocate %ux%ux%u in distance field atlas. Moved mesh distance field to FailedAllocations list"), Size.X, Size.Y, Size.Z);
+					UE_LOG(LogStaticMesh, Warning, TEXT("Failed to allocate %ux%ux%u for %s in distance field atlas. Moved mesh distance field to FailedAllocations list"), Size.X, Size.Y, Size.Z, *Texture->GetStaticMesh()->GetName());
 					LocalPendingAllocations->RemoveAt(AllocationIndex);
 					FailedAllocations.Add(Texture);
 					FailedAllocatedPixels += Size.X * Size.Y * Size.Z;
@@ -890,7 +890,7 @@ FDistanceFieldAsyncQueue* GDistanceFieldAsyncQueue = NULL;
 #if WITH_EDITOR
 
 // DDC key for distance field data, must be changed when modifying the generation code or data format
-#define DISTANCEFIELD_DERIVEDDATA_VER TEXT("CD4A6506-C64C-A019-AA56-2B0A412AE96A")
+#define DISTANCEFIELD_DERIVEDDATA_VER TEXT("CD4A6506-C64C-A029-AA56-2B0A412AE96A")
 
 FString BuildDistanceFieldDerivedDataKey(const FString& InMeshKey)
 {
@@ -954,15 +954,18 @@ void FDistanceFieldVolumeData::CacheDerivedData(const FString& InDDCKey, const I
 
 		for (int32 MaterialIndex = 0; MaterialIndex < Mesh->GetStaticMaterials().Num(); MaterialIndex++)
 		{
+			FSignedDistanceFieldBuildMaterialData MaterialData;
 			// Default material blend mode
-			EBlendMode BlendMode = BLEND_Opaque;
+			MaterialData.BlendMode = BLEND_Opaque;
+			MaterialData.bTwoSided = false;
 
 			if (Mesh->GetStaticMaterials()[MaterialIndex].MaterialInterface)
 			{
-				BlendMode = Mesh->GetStaticMaterials()[MaterialIndex].MaterialInterface->GetBlendMode();
+				MaterialData.BlendMode = Mesh->GetStaticMaterials()[MaterialIndex].MaterialInterface->GetBlendMode();
+				MaterialData.bTwoSided = Mesh->GetStaticMaterials()[MaterialIndex].MaterialInterface->IsTwoSided();
 			}
 
-			NewTask->MaterialBlendModes.Add(BlendMode);
+			NewTask->MaterialBlendModes.Add(MaterialData);
 		}
 
 		// Nanite overrides source static mesh with a coarse representation. Need to load original data before we build the mesh SDF.
diff --git a/Engine/Source/Runtime/Engine/Public/DistanceFieldAtlas.h b/Engine/Source/Runtime/Engine/Public/DistanceFieldAtlas.h
--- a/Engine/Source/Runtime/Engine/Public/DistanceFieldAtlas.h
+++ b/Engine/Source/Runtime/Engine/Public/DistanceFieldAtlas.h
@@ -18,6 +18,10 @@
 #include "DerivedMeshDataTaskUtils.h"
 #include "Async/AsyncWork.h"
 
+#if WITH_EDITOR
+#include "MeshUtilities.h"
+#endif
+
 class FDistanceFieldVolumeData;
 class UStaticMesh;
 class UTexture2D;
@@ -407,7 +411,9 @@ class FAsyncDistanceFieldTask
 public:
 	FAsyncDistanceFieldTask();
 
-	TArray<EBlendMode> MaterialBlendModes;
+#if WITH_EDITOR
+	TArray<FSignedDistanceFieldBuildMaterialData> MaterialBlendModes;
+#endif
 	FSourceMeshDataForDerivedDataTask SourceMeshData;
 	UStaticMesh* StaticMesh;
 	UStaticMesh* GenerateSource;
