Fixed Lumen card representation for thin planes, where decals were too thin to be picked up for projection

diff --git a/Engine/Source/Developer/MeshUtilities/Private/MeshCardRepresentationUtilities.cpp b/Engine/Source/Developer/MeshUtilities/Private/MeshCardRepresentationUtilities.cpp
--- a/Engine/Source/Developer/MeshUtilities/Private/MeshCardRepresentationUtilities.cpp
+++ b/Engine/Source/Developer/MeshUtilities/Private/MeshCardRepresentationUtilities.cpp
@@ -204,14 +204,13 @@ void SerializePlacedCards(TArray<FPlacedCard, TInlineAllocator<16>>& PlacedCards
 	}
 }
 
-void BuildMeshCards(const FBox& CardBounds, const FGenerateCardMeshContext& Context, FCardRepresentationData& OutData)
+void BuildMeshCards(const FBox& MeshBounds, const FGenerateCardMeshContext& Context, FCardRepresentationData& OutData)
 {
 	static const auto CVarMeshCardRepresentationMinSurface = IConsoleManager::Get().FindTConsoleVariableDataFloat(TEXT("r.MeshCardRepresentation.MinSurface"));
 	const float MinSurfaceThreshold = CVarMeshCardRepresentationMinSurface->GetValueOnAnyThread();
 
-	const float MaxExtent = CardBounds.GetExtent().GetMax();
-	// Ensure bounds don't have zero extent (was possible for planes)
-	const FBox MeshCardsBounds = CardBounds.ExpandBy(FVector::Max(CardBounds.GetExtent() * 0.05f, .001f * FVector(MaxExtent, MaxExtent, MaxExtent)));
+	// Make sure BBox isn't empty and we can generate card representation for it. This handles e.g. infinitely thin planes.
+	const FBox MeshCardsBounds = MeshBounds.ExpandBy(5.0f);
 
 	OutData.MeshCardsBuildData.Bounds = MeshCardsBounds;
 	OutData.MeshCardsBuildData.MaxLODLevel = 1;
diff --git a/Engine/Source/Runtime/Engine/Private/MeshCardRepresentation.cpp b/Engine/Source/Runtime/Engine/Private/MeshCardRepresentation.cpp
--- a/Engine/Source/Runtime/Engine/Private/MeshCardRepresentation.cpp
+++ b/Engine/Source/Runtime/Engine/Private/MeshCardRepresentation.cpp
@@ -60,7 +60,7 @@ FCardRepresentationAsyncQueue* GCardRepresentationAsyncQueue = NULL;
 #if WITH_EDITOR
 
 // DDC key for card representation data, must be changed when modifying the generation code or data format
-#define CARDREPRESENTATION_DERIVEDDATA_VER TEXT("C1F7A92371F1A62CB3A14C27C7B4014A")
+#define CARDREPRESENTATION_DERIVEDDATA_VER TEXT("26994DB0CF1F41C5817CB23A63F13B78")
 
 FString BuildCardRepresentationDerivedDataKey(const FString& InMeshKey)
 {
