Optim: reserve proper array size in FSkeletalMeshUtilityBuilder::GenerateSkeletalRenderMesh to avoid unrequired memory allocations.

diff --git a/Engine/Source/Developer/MeshUtilities/Private/MeshUtilities.cpp b/Engine/Source/Developer/MeshUtilities/Private/MeshUtilities.cpp
--- a/Engine/Source/Developer/MeshUtilities/Private/MeshUtilities.cpp
+++ b/Engine/Source/Developer/MeshUtilities/Private/MeshUtilities.cpp
@@ -4012,8 +4012,8 @@ public:
 		TArray<FSkeletalMeshVertIndexAndZ> VertIndexAndZ;
 		TArray<FSoftSkinBuildVertex> RawVertices;
 
-		VertIndexAndZ.Empty(BuildData.Points.Num());
-		RawVertices.Reserve(BuildData.Points.Num());
+		VertIndexAndZ.Empty(BuildData.Faces.Num() * 3);
+		RawVertices.Reserve(BuildData.Faces.Num() * 3);
 
 		for (int32 FaceIndex = 0; FaceIndex < BuildData.Faces.Num(); FaceIndex++)
 		{
