U5 - Move FMeshDistanceFieldAsyncTask in parallelFor

diff --git a/Engine/Source/Developer/MeshUtilities/Private/MeshDistanceFieldUtilities.cpp b/Engine/Source/Developer/MeshUtilities/Private/MeshDistanceFieldUtilities.cpp
--- a/Engine/Source/Developer/MeshUtilities/Private/MeshDistanceFieldUtilities.cpp
+++ b/Engine/Source/Developer/MeshUtilities/Private/MeshDistanceFieldUtilities.cpp
@@ -9,11 +9,12 @@
 #include "StaticMeshResources.h"
 #include "DistanceFieldAtlas.h"
 #include "MeshRepresentationCommon.h"
+#include "Async/ParallelFor.h"
 
 //@todo - implement required vector intrinsics for other implementations
 #if PLATFORM_ENABLE_VECTORINTRINSICS
 
-class FMeshDistanceFieldAsyncTask : public FNonAbandonableTask
+class FMeshDistanceFieldAsyncTask
 {
 public:
 	FMeshDistanceFieldAsyncTask(
@@ -41,11 +42,6 @@ public:
 
 	void DoWork();
 
-	FORCEINLINE TStatId GetStatId() const
-	{
-		RETURN_QUICK_DECLARE_CYCLE_STAT(FMeshDistanceFieldAsyncTask, STATGROUP_ThreadPoolAsyncTasks);
-	}
-
 	bool WasNegativeAtBorder() const
 	{
 		return bNegativeAtBorder;
@@ -279,11 +275,12 @@ void FMeshUtilities::GenerateSignedDistanceFieldVolumeData(
 			DistanceFieldVolume.Empty(VolumeDimensions.X * VolumeDimensions.Y * VolumeDimensions.Z);
 			DistanceFieldVolume.AddZeroed(VolumeDimensions.X * VolumeDimensions.Y * VolumeDimensions.Z);
 
-			TIndirectArray<FAsyncTask<FMeshDistanceFieldAsyncTask>> AsyncTasks;
+			TArray<FMeshDistanceFieldAsyncTask> AsyncTasks;
+			AsyncTasks.Reserve(VolumeDimensions.Z);
 
 			for (int32 ZIndex = 0; ZIndex < VolumeDimensions.Z; ZIndex++)
 			{
-				FAsyncTask<FMeshDistanceFieldAsyncTask>* Task = new FAsyncTask<class FMeshDistanceFieldAsyncTask>(
+				AsyncTasks.Emplace(
 					&EmbreeScene.kDopTree,
 					EmbreeScene.bUseEmbree,
 					EmbreeScene.EmbreeScene,
@@ -293,21 +290,18 @@ void FMeshUtilities::GenerateSignedDistanceFieldVolumeData(
 					DistanceFieldVolumeMaxDistance,
 					ZIndex,
 					&DistanceFieldVolume);
-
-				//Task->StartSynchronousTask();
-				Task->StartBackgroundTask(&ThreadPool);
-
-				AsyncTasks.Add(Task);
 			}
 
 			bool bNegativeAtBorder = false;
 
+			ParallelForTemplate(AsyncTasks.Num(), [&AsyncTasks](int32 TaskIndex)
+			{
+				AsyncTasks[TaskIndex].DoWork();
+			}, EParallelForFlags::BackgroundPriority | EParallelForFlags::Unbalanced);
+
 			for (int32 TaskIndex = 0; TaskIndex < AsyncTasks.Num(); TaskIndex++)
 			{
-				FAsyncTask<FMeshDistanceFieldAsyncTask>& Task = AsyncTasks[TaskIndex];
-				extern CORE_API int32 GUseNewTaskBackend;
-				Task.EnsureCompletion(!!GUseNewTaskBackend);
-				bNegativeAtBorder = bNegativeAtBorder || Task.GetTask().WasNegativeAtBorder();
+				bNegativeAtBorder = bNegativeAtBorder || AsyncTasks[TaskIndex].WasNegativeAtBorder();
 			}
 
 			if (bNegativeAtBorder)
