Lumen - tighten surface cache projections to save a bit of atlas space and remove empty borders causing incorrect radiosity upsampling

diff --git a/Engine/Source/Developer/MeshUtilities/Private/MeshCardRepresentationUtilities.cpp b/Engine/Source/Developer/MeshUtilities/Private/MeshCardRepresentationUtilities.cpp
--- a/Engine/Source/Developer/MeshUtilities/Private/MeshCardRepresentationUtilities.cpp
+++ b/Engine/Source/Developer/MeshUtilities/Private/MeshCardRepresentationUtilities.cpp
@@ -185,6 +185,7 @@ void SerializePlacedCards(TArray<FPlacedCard, TInlineAllocator<16>>& PlacedCards
 	int32 LODLevel,
 	int32 Orientation,
 	int32 MinCardHits,
+	const FBox& MeshCardsBounds,
 	FCardRepresentationData& OutData)
 {
 	for (int32 PlacedCardIndex = 0; PlacedCardIndex < PlacedCards.Num(); ++PlacedCardIndex)
@@ -192,9 +193,11 @@ void SerializePlacedCards(TArray<FPlacedCard, TInlineAllocator<16>>& PlacedCards
 		const FPlacedCard& PlacedCard = PlacedCards[PlacedCardIndex];
 		if (PlacedCard.NumHits >= MinCardHits)
 		{
+			const FBox ClampedBox = PlacedCard.Bounds.Overlap(MeshCardsBounds);
+
 			FLumenCardBuildData CardBuildData;
-			CardBuildData.Center = PlacedCard.Bounds.GetCenter();
-			CardBuildData.Extent = PlacedCard.Bounds.GetExtent();
+			CardBuildData.Center = ClampedBox.GetCenter();
+			CardBuildData.Extent = ClampedBox.GetExtent();
 			CardBuildData.Extent = TransformFaceExtent(CardBuildData.Extent, Orientation);
 			CardBuildData.Orientation = Orientation;
 			CardBuildData.LODLevel = LODLevel;
@@ -210,7 +213,9 @@ void BuildMeshCards(const FBox& MeshBounds, const FGenerateCardMeshContext& Cont
 	const float MinSurfaceThreshold = CVarMeshCardRepresentationMinSurface->GetValueOnAnyThread();
 
 	// Make sure BBox isn't empty and we can generate card representation for it. This handles e.g. infinitely thin planes.
-	const FBox MeshCardsBounds = MeshBounds.ExpandBy(5.0f);
+	const FVector MeshCardsBoundsCenter = MeshBounds.GetCenter();
+	const FVector MeshCardsBoundsExtent = FVector::Max(MeshBounds.GetExtent() + 1.0f, FVector(5.0f));
+	const FBox MeshCardsBounds(MeshCardsBoundsCenter - MeshCardsBoundsExtent, MeshCardsBoundsCenter + MeshCardsBoundsExtent);
 
 	OutData.MeshCardsBuildData.Bounds = MeshCardsBounds;
 	OutData.MeshCardsBuildData.MaxLODLevel = 1;
@@ -250,6 +255,8 @@ void BuildMeshCards(const FBox& MeshBounds, const FGenerateCardMeshContext& Cont
 				MeshSliceNum = VolumeSizeInVoxels.X;
 				HeighfieldSize.X = VolumeSizeInVoxels.Y;
 				HeighfieldSize.Y = VolumeSizeInVoxels.Z;
+				HeighfieldStepX = FVector(0.0f, MeshCardsBounds.GetSize().Y / HeighfieldSize.X, 0.0f);
+				HeighfieldStepY = FVector(0.0f, 0.0f, MeshCardsBounds.GetSize().Z / HeighfieldSize.Y);
 				break;
 
 			case 1:
@@ -257,6 +264,8 @@ void BuildMeshCards(const FBox& MeshBounds, const FGenerateCardMeshContext& Cont
 				MeshSliceNum = VolumeSizeInVoxels.Y;
 				HeighfieldSize.X = VolumeSizeInVoxels.X;
 				HeighfieldSize.Y = VolumeSizeInVoxels.Z;
+				HeighfieldStepX = FVector(MeshCardsBounds.GetSize().X / HeighfieldSize.X, 0.0f, 0.0f);
+				HeighfieldStepY = FVector(0.0f, 0.0f, MeshCardsBounds.GetSize().Z / HeighfieldSize.Y);
 				break;
 
 			case 2:
@@ -264,6 +273,8 @@ void BuildMeshCards(const FBox& MeshBounds, const FGenerateCardMeshContext& Cont
 				MeshSliceNum = VolumeSizeInVoxels.Z;
 				HeighfieldSize.X = VolumeSizeInVoxels.X;
 				HeighfieldSize.Y = VolumeSizeInVoxels.Y;
+				HeighfieldStepX = FVector(MeshCardsBounds.GetSize().X / HeighfieldSize.X, 0.0f, 0.0f);
+				HeighfieldStepY = FVector(0.0f, MeshCardsBounds.GetSize().Y / HeighfieldSize.Y, 0.0f);
 				break;
 		}
 
@@ -271,41 +282,29 @@ void BuildMeshCards(const FBox& MeshBounds, const FGenerateCardMeshContext& Cont
 		{
 			case 0: 
 				RayDirection.X = +1.0f; 
-				HeighfieldStepX = FVector(0.0f, MeshCardsBounds.GetSize().Y / HeighfieldSize.X, 0.0f);
-				HeighfieldStepY = FVector(0.0f, 0.0f, MeshCardsBounds.GetSize().Z / HeighfieldSize.Y);
 				break;
 
 			case 1: 
 				RayDirection.X = -1.0f; 
 				RayOriginFrame.X = MeshCardsBounds.Max.X;
-				HeighfieldStepX = FVector(0.0f, MeshCardsBounds.GetSize().Y / HeighfieldSize.X, 0.0f);
-				HeighfieldStepY = FVector(0.0f, 0.0f, MeshCardsBounds.GetSize().Z / HeighfieldSize.Y);
 				break;
 
 			case 2: 
 				RayDirection.Y = +1.0f; 
-				HeighfieldStepX = FVector(MeshCardsBounds.GetSize().X / HeighfieldSize.X, 0.0f, 0.0f);
-				HeighfieldStepY = FVector(0.0f, 0.0f, MeshCardsBounds.GetSize().Z / HeighfieldSize.Y);
 				break;
 
 			case 3: 
 				RayDirection.Y = -1.0f; 
 				RayOriginFrame.Y = MeshCardsBounds.Max.Y;
-				HeighfieldStepX = FVector(MeshCardsBounds.GetSize().X / HeighfieldSize.X, 0.0f, 0.0f);
-				HeighfieldStepY = FVector(0.0f, 0.0f, MeshCardsBounds.GetSize().Z / HeighfieldSize.Y);
 				break;
 
 			case 4: 
 				RayDirection.Z = +1.0f; 
-				HeighfieldStepX = FVector(MeshCardsBounds.GetSize().X / HeighfieldSize.X, 0.0f, 0.0f);
-				HeighfieldStepY = FVector(0.0f, MeshCardsBounds.GetSize().Y / HeighfieldSize.Y, 0.0f);
 				break;
 
 			case 5: 
 				RayDirection.Z = -1.0f; 
 				RayOriginFrame.Z = MeshCardsBounds.Max.Z;
-				HeighfieldStepX = FVector(MeshCardsBounds.GetSize().X / HeighfieldSize.X, 0.0f, 0.0f);
-				HeighfieldStepY = FVector(0.0f, MeshCardsBounds.GetSize().Y / HeighfieldSize.Y, 0.0f);
 				break;
 
 			default: 
@@ -418,7 +417,7 @@ void BuildMeshCards(const FBox& MeshBounds, const FGenerateCardMeshContext& Cont
 			}
 		}
 
-		SerializePlacedCards(PlacedCards, /*LOD level*/ 0, Orientation, MinCardHits, OutData);
+		SerializePlacedCards(PlacedCards, /*LOD level*/ 0, Orientation, MinCardHits, MeshCardsBounds, OutData);
 
 		// Try to place more cards by splitting existing ones
 		for (uint32 CardPlacementIteration = 0; CardPlacementIteration < 4; ++CardPlacementIteration)
@@ -468,7 +467,7 @@ void BuildMeshCards(const FBox& MeshBounds, const FGenerateCardMeshContext& Cont
 			}
 		}
 
-		SerializePlacedCards(PlacedCards, /*LOD level*/ 1, Orientation, MinCardHits, OutData);
+		SerializePlacedCards(PlacedCards, /*LOD level*/ 1, Orientation, MinCardHits, MeshCardsBounds, OutData);
 	}
 }
 
diff --git a/Engine/Source/Runtime/Engine/Private/MeshCardRepresentation.cpp b/Engine/Source/Runtime/Engine/Private/MeshCardRepresentation.cpp
--- a/Engine/Source/Runtime/Engine/Private/MeshCardRepresentation.cpp
+++ b/Engine/Source/Runtime/Engine/Private/MeshCardRepresentation.cpp
@@ -60,7 +60,7 @@ FCardRepresentationAsyncQueue* GCardRepresentationAsyncQueue = NULL;
 #if WITH_EDITOR
 
 // DDC key for card representation data, must be changed when modifying the generation code or data format
-#define CARDREPRESENTATION_DERIVEDDATA_VER TEXT("26994DB0CF1F41C5817CB23A63F13B78")
+#define CARDREPRESENTATION_DERIVEDDATA_VER TEXT("E8C7CF855B174FE195BE69EAF600F293")
 
 FString BuildCardRepresentationDerivedDataKey(const FString& InMeshKey)
 {
