ParallelFor over orientations when building mesh card representation.

diff --git a/Engine/Source/Developer/MeshUtilities/Private/MeshCardRepresentationUtilities.cpp b/Engine/Source/Developer/MeshUtilities/Private/MeshCardRepresentationUtilities.cpp
--- a/Engine/Source/Developer/MeshUtilities/Private/MeshCardRepresentationUtilities.cpp
+++ b/Engine/Source/Developer/MeshUtilities/Private/MeshCardRepresentationUtilities.cpp
@@ -165,7 +165,7 @@ int32 UpdatePlacedCards(TArray<FPlacedCard, TInlineAllocator<16>>& PlacedCards,
 	return NumMeshHits;
 }
 
-void SerializePlacedCards(TArray<FPlacedCard, TInlineAllocator<16>>& PlacedCards, 
+void SerializePlacedCards(TArrayView<const FPlacedCard> PlacedCards,
 	int32 LODLevel,
 	int32 Orientation,
 	int32 MinCardHits,
@@ -221,8 +221,22 @@ void BuildMeshCards(const FBox& MeshBounds, const FGenerateCardMeshContext& Cont
 		FRandomStream RandomStream(0);
 		MeshUtilities::GenerateStratifiedUniformHemisphereSamples(64, RandomStream, RayDirectionsOverHemisphere);
 	}
-		
-	for (int32 Orientation = 0; Orientation < 6; ++Orientation)
+
+	using FPlacedCardArray = TArray<FPlacedCard, TInlineAllocator<16>>;
+	struct FTaskOutputs
+	{
+		FPlacedCardArray PlacedCardsPerLod[2];
+		float MinCardHitsPerLod[2] = {};
+	};
+	
+	FTaskOutputs TaskOutputsPerOrientation[6];
+
+	ParallelFor(6, [VolumeSizeInVoxels, MinSurfaceThreshold, VoxelExtent,
+		&Context = AsConst(Context),
+		&RayDirectionsOverHemisphere = AsConst(RayDirectionsOverHemisphere),
+		&MeshCardsBounds = AsConst(MeshCardsBounds),
+		&TaskOutputsPerOrientation
+	] (int32 Orientation)
 	{
 		FIntPoint HeighfieldSize(0, 0);
 		FVector RayDirection(0.0f, 0.0f, 0.0f);
@@ -369,11 +383,9 @@ void BuildMeshCards(const FBox& MeshBounds, const FGenerateCardMeshContext& Cont
 			}
 		}
 
-
 		const int32 MinCardHits = FMath::Floor(HeighfieldSize.X * HeighfieldSize.Y * MinSurfaceThreshold);
 
-
-		TArray<FPlacedCard, TInlineAllocator<16>> PlacedCards;
+		FPlacedCardArray& PlacedCardsLod0 = TaskOutputsPerOrientation[Orientation].PlacedCardsPerLod[0];
 		int32 PlacedCardsHits = 0;
 
 		// Place a default card
@@ -381,9 +393,9 @@ void BuildMeshCards(const FBox& MeshBounds, const FGenerateCardMeshContext& Cont
 			FPlacedCard PlacedCard;
 			PlacedCard.SliceMin = 0;
 			PlacedCard.SliceMax = MeshSliceNum;
-			PlacedCards.Add(PlacedCard);
+			PlacedCardsLod0.Add(PlacedCard);
 
-			PlacedCardsHits = UpdatePlacedCards(PlacedCards,
+			PlacedCardsHits = UpdatePlacedCards(PlacedCardsLod0,
 				RayOriginFrame,
 				RayDirection,
 				HeighfieldStepX,
@@ -397,11 +409,14 @@ void BuildMeshCards(const FBox& MeshBounds, const FGenerateCardMeshContext& Cont
 
 			if (PlacedCardsHits < MinCardHits)
 			{
-				PlacedCards.Reset();
+				PlacedCardsLod0.Reset();
 			}
 		}
 
-		SerializePlacedCards(PlacedCards, /*LOD level*/ 0, Orientation, MinCardHits, MeshCardsBounds, OutData);
+		TaskOutputsPerOrientation[Orientation].MinCardHitsPerLod[0] = MinCardHits;
+
+		FPlacedCardArray& PlacedCardsLod1 = TaskOutputsPerOrientation[Orientation].PlacedCardsPerLod[1];
+		PlacedCardsLod1 = PlacedCardsLod0;
 
 		// Try to place more cards by splitting existing ones
 		for (uint32 CardPlacementIteration = 0; CardPlacementIteration < 4; ++CardPlacementIteration)
@@ -409,12 +424,12 @@ void BuildMeshCards(const FBox& MeshBounds, const FGenerateCardMeshContext& Cont
 			TArray<FPlacedCard, TInlineAllocator<16>> BestPlacedCards;
 			int32 BestPlacedCardHits = PlacedCardsHits;
 
-			for (int32 PlacedCardIndex = 0; PlacedCardIndex < PlacedCards.Num(); ++PlacedCardIndex)
+			for (int32 PlacedCardIndex = 0; PlacedCardIndex < PlacedCardsLod1.Num(); ++PlacedCardIndex)
 			{
-				const FPlacedCard& PlacedCard = PlacedCards[PlacedCardIndex];
+				const FPlacedCard& PlacedCard = PlacedCardsLod1[PlacedCardIndex];
 				for (int32 SliceIndex = PlacedCard.SliceMin + 2; SliceIndex < PlacedCard.SliceMax; ++SliceIndex)
 				{
-					TArray<FPlacedCard, TInlineAllocator<16>> TempPlacedCards(PlacedCards);
+					TArray<FPlacedCard, TInlineAllocator<16>> TempPlacedCards(PlacedCardsLod1);
 
 					FPlacedCard NewPlacedCard;
 					NewPlacedCard.SliceMin = SliceIndex;
@@ -446,12 +461,19 @@ void BuildMeshCards(const FBox& MeshBounds, const FGenerateCardMeshContext& Cont
 
 			if (BestPlacedCardHits >= PlacedCardsHits + MinCardHits)
 			{
-				PlacedCards = BestPlacedCards;
+				PlacedCardsLod1 = BestPlacedCards;
 				PlacedCardsHits = BestPlacedCardHits;
 			}
 		}
 
-		SerializePlacedCards(PlacedCards, /*LOD level*/ 1, Orientation, MinCardHits, MeshCardsBounds, OutData);
+		TaskOutputsPerOrientation[Orientation].MinCardHitsPerLod[1] = MinCardHits;
+	});
+
+	for (int32 Orientation = 0; Orientation < 6; ++Orientation)
+	{
+		const FTaskOutputs& TaskOutputs = TaskOutputsPerOrientation[Orientation];
+		SerializePlacedCards(TaskOutputs.PlacedCardsPerLod[0], /*LOD level*/ 0, Orientation, TaskOutputs.MinCardHitsPerLod[0], MeshCardsBounds, OutData);
+		SerializePlacedCards(TaskOutputs.PlacedCardsPerLod[1], /*LOD level*/ 1, Orientation, TaskOutputs.MinCardHitsPerLod[1], MeshCardsBounds, OutData);
 	}
 }
 
