Avoid copying libtbb.so for Linux and just depend on its actual location

diff --git a/Engine/Plugins/Runtime/ResonanceAudio/Source/ResonanceAudio/ResonanceAudio.Build.cs b/Engine/Plugins/Runtime/ResonanceAudio/Source/ResonanceAudio/ResonanceAudio.Build.cs
--- a/Engine/Plugins/Runtime/ResonanceAudio/Source/ResonanceAudio/ResonanceAudio.Build.cs
+++ b/Engine/Plugins/Runtime/ResonanceAudio/Source/ResonanceAudio/ResonanceAudio.Build.cs
@@ -2,6 +2,7 @@
 // Copyright (C) Google Inc. 2017. All rights reserved.
 //
 using UnrealBuildTool;
+using System.IO;
 
 public class ResonanceAudio : ModuleRules
 {
@@ -118,12 +119,14 @@ public class ResonanceAudio : ModuleRules
 		{
 		   string SDKDir = Target.UEThirdPartySourceDirectory + "Intel/Embree/Embree2140/Linux/x86_64-unknown-linux-gnu/";
 
-            PublicIncludePaths.Add(SDKDir + "include");
-            PublicAdditionalLibraries.Add(SDKDir + "lib/libembree.so");
-            RuntimeDependencies.Add("$(TargetOutputDir)/libembree.so", SDKDir + "lib/libembree.so");
-            RuntimeDependencies.Add("$(TargetOutputDir)/libtbb.so", IntelTBBLibs + "Linux/libtbb.so");
+            PublicIncludePaths.Add(Path.Combine(SDKDir, "include"));
+            PublicAdditionalLibraries.Add(Path.Combine(SDKDir, "lib/libembree.so"));
+			PublicAdditionalLibraries.Add(Path.Combine(IntelTBBLibs, "Linux/libtbb.so"));
+			RuntimeDependencies.Add(Path.Combine(IntelTBBLibs, "Linux/libtbb.so"));
+			RuntimeDependencies.Add(Path.Combine(IntelTBBLibs, "Linux/libtbb.so.2"));
+
 			// disabled for Linux atm due to a bug in libtbbmalloc on exit
-            // RuntimeDependencies.Add("$(TargetOutputDir)/libtbbmalloc.so", IntelTBBLibs + "Linux/libtbbmalloc.so");
+			// RuntimeDependencies.Add(Path.Combine(IntelTBBLibs, "Linux/libtbbmalloc.so"));
 			PublicDefinitions.Add("USE_EMBREE=1");
 			PrivateDefinitions.Add("USE_EMBRE_FOR_RESONANCE=1");
 		}
diff --git a/Engine/Source/Developer/MeshUtilities/MeshUtilities.Build.cs b/Engine/Source/Developer/MeshUtilities/MeshUtilities.Build.cs
--- a/Engine/Source/Developer/MeshUtilities/MeshUtilities.Build.cs
+++ b/Engine/Source/Developer/MeshUtilities/MeshUtilities.Build.cs
@@ -112,15 +112,15 @@ public class MeshUtilities : ModuleRules
 		{
 			string SDKDir = Target.UEThirdPartySourceDirectory + "Intel/Embree/Embree2140/Linux/x86_64-unknown-linux-gnu/";
 
-            PublicIncludePaths.Add(SDKDir + "include");
-            PublicAdditionalLibraries.Add(SDKDir + "lib/libembree.so");
-            PublicAdditionalLibraries.Add(IntelTBBLibs + "Linux/libtbb.so");
+            PublicIncludePaths.Add(Path.Combine(SDKDir, "include"));
+            PublicAdditionalLibraries.Add(Path.Combine(SDKDir, "lib/libembree.so"));
+            PublicAdditionalLibraries.Add(Path.Combine(IntelTBBLibs, "Linux/libtbb.so"));
 			// disabled for Linux atm due to a bug in libtbbmalloc on exit
-            // PublicAdditionalLibraries.Add(IntelTBBLibs + "Linux/libtbbmalloc.so");
-            RuntimeDependencies.Add("$(TargetOutputDir)/libembree.so", SDKDir + "lib/libembree.so");
-            RuntimeDependencies.Add("$(TargetOutputDir)/libtbb.so", IntelTBBLibs + "Linux/libtbb.so");
+            // PublicAdditionalLibraries.Add(Path.Combine(IntelTBBLibs, "Linux/libtbbmalloc.so"));
+			RuntimeDependencies.Add(Path.Combine(IntelTBBLibs, "Linux/libtbb.so"));
+			RuntimeDependencies.Add(Path.Combine(IntelTBBLibs, "Linux/libtbb.so.2"));
 			// disabled for Linux atm due to a bug in libtbbmalloc on exit
-            // RuntimeDependencies.Add("$(TargetOutputDir)/libtbbmalloc.so", IntelTBBLibs + "Linux/libtbbmalloc.so");
+            // RuntimeDependencies.Add(Path.Combine(IntelTBBLibs, "Linux/libtbbmalloc.so"));
             PublicDefinitions.Add("USE_EMBREE=1");
 		}
         else
diff --git a/Engine/Source/Programs/UnrealLightmass/UnrealLightmass.Build.cs b/Engine/Source/Programs/UnrealLightmass/UnrealLightmass.Build.cs
--- a/Engine/Source/Programs/UnrealLightmass/UnrealLightmass.Build.cs
+++ b/Engine/Source/Programs/UnrealLightmass/UnrealLightmass.Build.cs
@@ -1,6 +1,7 @@
 // Copyright Epic Games, Inc. All Rights Reserved.
 
 using UnrealBuildTool;
+using System.IO;
 
 public class UnrealLightmass : ModuleRules
 {
@@ -86,10 +87,12 @@ public class UnrealLightmass : ModuleRules
 		{
             string SDKDir = Target.UEThirdPartySourceDirectory + "Intel/Embree/Embree270/Linux/x86_64-unknown-linux-gnu/";
 
-            PublicIncludePaths.Add(SDKDir + "include");
-            PublicAdditionalLibraries.Add(SDKDir + "lib/libembree.so");
+            PublicIncludePaths.Add(Path.Combine(SDKDir, "include"));
+            PublicAdditionalLibraries.Add(Path.Combine(SDKDir, "lib/libembree.so"));
 			RuntimeDependencies.Add("$(EngineDir)/Binaries/Linux/libembree.so");
-			RuntimeDependencies.Add("$(EngineDir)/Binaries/Linux/libtbb.so", IntelTBBLibs + "Linux/libtbb.so"); // Take latest version to avoid overwriting the editor's copy
+			PublicAdditionalLibraries.Add(Path.Combine(IntelTBBLibs, "Linux/libtbb.so"));
+			RuntimeDependencies.Add(Path.Combine(IntelTBBLibs, "Linux/libtbb.so"));
+			RuntimeDependencies.Add(Path.Combine(IntelTBBLibs, "Linux/libtbb.so.2"));
 			// disabled for Linux atm due to a bug in libtbbmalloc on exit
 			// RuntimeDependencies.Add("$(EngineDir)/Binaries/Linux/libtbbmalloc.so", IntelTBBLibs + "Linux/libtbbmalloc.so");
 			PublicDefinitions.Add("USE_EMBREE=1");
