Fixed Linux compile error

diff --git a/Engine/Source/Developer/MeshUtilities/Private/MeshRepresentationCommon.h b/Engine/Source/Developer/MeshUtilities/Private/MeshRepresentationCommon.h
--- a/Engine/Source/Developer/MeshUtilities/Private/MeshRepresentationCommon.h
+++ b/Engine/Source/Developer/MeshUtilities/Private/MeshRepresentationCommon.h
@@ -11,6 +11,7 @@
 #else
 	typedef void* RTCDevice;
 	typedef void* RTCScene;
+	typedef void* RTCGeometry;
 #endif
 
 class FSourceMeshDataForDerivedDataTask;
