Fixed crash cooking when project does not have r.GenerateMeshDistanceFields enabled

diff --git a/Engine/Source/Developer/MeshUtilities/Private/MeshCardRepresentationUtilities.cpp b/Engine/Source/Developer/MeshUtilities/Private/MeshCardRepresentationUtilities.cpp
--- a/Engine/Source/Developer/MeshUtilities/Private/MeshCardRepresentationUtilities.cpp
+++ b/Engine/Source/Developer/MeshUtilities/Private/MeshCardRepresentationUtilities.cpp
@@ -515,7 +515,7 @@ bool FMeshUtilities::GenerateCardRepresentationData(
 	FGenerateCardMeshContext Context(MeshName, EmbreeScene.EmbreeScene, EmbreeScene.EmbreeDevice, OutData);
 
 	// Note: must operate on the SDF bounds because SDF generation can expand the mesh's bounds
-	BuildMeshCards(DistanceFieldVolumeData->LocalSpaceMeshBounds, Context, OutData);
+	BuildMeshCards(DistanceFieldVolumeData ? DistanceFieldVolumeData->LocalSpaceMeshBounds : Bounds.GetBox(), Context, OutData);
 
 	MeshRepresentation::DeleteEmbreeScene(EmbreeScene);
 
