Changed lack of Embree to be a warning, not error log, to not fail DDC build commandlet

diff --git a/Engine/Source/Developer/MeshUtilities/Private/MeshDistanceFieldUtilities.cpp b/Engine/Source/Developer/MeshUtilities/Private/MeshDistanceFieldUtilities.cpp
--- a/Engine/Source/Developer/MeshUtilities/Private/MeshDistanceFieldUtilities.cpp
+++ b/Engine/Source/Developer/MeshUtilities/Private/MeshDistanceFieldUtilities.cpp
@@ -549,7 +549,7 @@ void FMeshUtilities::GenerateSignedDistanceFieldVolumeData(
 {
 	if (DistanceFieldResolutionScale > 0)
 	{
-		UE_LOG(LogMeshUtilities, Error, TEXT("Couldn't generate distance field for mesh, platform is missing Embree support."));
+		UE_LOG(LogMeshUtilities, Warning, TEXT("Couldn't generate distance field for mesh, platform is missing Embree support."));
 	}
 }
 
