Removed external linkage from ComputerLinearVoxelIndex() in two separate modules to prevent monolithic link errors.

#rnx

diff --git a/Engine/Source/Developer/MeshUtilities/Private/MeshDistanceFieldUtilities.cpp b/Engine/Source/Developer/MeshUtilities/Private/MeshDistanceFieldUtilities.cpp
--- a/Engine/Source/Developer/MeshUtilities/Private/MeshDistanceFieldUtilities.cpp
+++ b/Engine/Source/Developer/MeshUtilities/Private/MeshDistanceFieldUtilities.cpp
@@ -64,7 +64,7 @@ bool EmbreePointQueryFunction(RTCPointQueryFunctionArguments* args)
 	return false;
 }
 
-int32 ComputeLinearVoxelIndex(FIntVector VoxelCoordinate, FIntVector VolumeDimensions)
+static int32 ComputeLinearVoxelIndex(FIntVector VoxelCoordinate, FIntVector VolumeDimensions)
 {
 	return (VoxelCoordinate.Z * VolumeDimensions.Y + VoxelCoordinate.Y) * VolumeDimensions.X + VoxelCoordinate.X;
 }
diff --git a/Engine/Source/Editor/UnrealEd/Private/Lightmass/ImportVolumetricLightmap.cpp b/Engine/Source/Editor/UnrealEd/Private/Lightmass/ImportVolumetricLightmap.cpp
--- a/Engine/Source/Editor/UnrealEd/Private/Lightmass/ImportVolumetricLightmap.cpp
+++ b/Engine/Source/Editor/UnrealEd/Private/Lightmass/ImportVolumetricLightmap.cpp
@@ -76,7 +76,7 @@ void CopyBetweenAtlasVolumeTextures(
 	}
 }
 
-int32 ComputeLinearVoxelIndex(FIntVector VoxelCoordinate, FIntVector VolumeDimensions)
+static int32 ComputeLinearVoxelIndex(FIntVector VoxelCoordinate, FIntVector VolumeDimensions)
 {
 	return (VoxelCoordinate.Z * VolumeDimensions.Y + VoxelCoordinate.Y) * VolumeDimensions.X + VoxelCoordinate.X;
 }
