Fix the function CreateImportDataFromLODModel and add some automation test to make sure the function is working properly.

diff --git a/Engine/Source/Developer/MeshUtilities/Private/MeshUtilities.cpp b/Engine/Source/Developer/MeshUtilities/Private/MeshUtilities.cpp
--- a/Engine/Source/Developer/MeshUtilities/Private/MeshUtilities.cpp
+++ b/Engine/Source/Developer/MeshUtilities/Private/MeshUtilities.cpp
@@ -5903,9 +5903,8 @@ void FMeshUtilities::CreateImportDataFromLODModel(USkeletalMesh* SkeletalMesh) c
 		for (int32 BoneIndex = 0; BoneIndex < ReferenceSkeleton.GetRawBoneNum(); ++BoneIndex)
 		{
 			SkeletalMeshImportData::FBone& RefBone = ImportData.RefBonesBinary.AddDefaulted_GetRef();
-			const FBoneIndexType& RequiredBone = TempLODModel.RequiredBones[BoneIndex];
-			const FMeshBoneInfo& MeshBoneInfo = ReferenceSkeleton.GetRawRefBoneInfo()[RequiredBone];
-			const FTransform& MeshBonePose = ReferenceSkeleton.GetRawRefBonePose()[RequiredBone];
+			const FMeshBoneInfo& MeshBoneInfo = ReferenceSkeleton.GetRawRefBoneInfo()[BoneIndex];
+			const FTransform& MeshBonePose = ReferenceSkeleton.GetRawRefBonePose()[BoneIndex];
 			RefBone.Name = MeshBoneInfo.ExportName;
 			RefBone.BonePos.Transform = MeshBonePose;
 			RefBone.ParentIndex = MeshBoneInfo.ParentIndex;
