Fix usage of embree with LWC for mesh distance fields and cards

diff --git a/Engine/Source/Developer/MeshBuilder/Private/MeshBuilderModule.cpp b/Engine/Source/Developer/MeshBuilder/Private/MeshBuilderModule.cpp
--- a/Engine/Source/Developer/MeshBuilder/Private/MeshBuilderModule.cpp
+++ b/Engine/Source/Developer/MeshBuilder/Private/MeshBuilderModule.cpp
@@ -31,7 +31,7 @@ public:
 	virtual bool BuildMeshVertexPositions(
 		UObject* StaticMesh,
 		TArray<uint32>& Indices,
-		TArray<FVector>& Vertices) override;
+		TArray<FVector3f>& Vertices) override;
 
 	virtual bool BuildSkeletalMesh(const FSkeletalMeshBuildParameters& SkeletalMeshBuildParameters) override;
 
@@ -55,7 +55,7 @@ bool FMeshBuilderModule::BuildMesh(FStaticMeshRenderData& OutRenderData, class U
 bool FMeshBuilderModule::BuildMeshVertexPositions(
 	UObject* Mesh,
 	TArray<uint32>& Indices,
-	TArray<FVector>& Vertices)
+	TArray<FVector3f>& Vertices)
 {
 	UStaticMesh* StaticMesh = Cast< UStaticMesh >(Mesh);
 	if (StaticMesh)
diff --git a/Engine/Source/Developer/MeshBuilder/Private/StaticMeshBuilder.cpp b/Engine/Source/Developer/MeshBuilder/Private/StaticMeshBuilder.cpp
--- a/Engine/Source/Developer/MeshBuilder/Private/StaticMeshBuilder.cpp
+++ b/Engine/Source/Developer/MeshBuilder/Private/StaticMeshBuilder.cpp
@@ -570,7 +570,7 @@ bool FStaticMeshBuilder::Build(FStaticMeshRenderData& StaticMeshRenderData, USta
 bool FStaticMeshBuilder::BuildMeshVertexPositions(
 	UStaticMesh* StaticMesh,
 	TArray<uint32>& BuiltIndices,
-	TArray<FVector>& BuiltVertices)
+	TArray<FVector3f>& BuiltVertices)
 {
 	TRACE_CPUPROFILER_EVENT_SCOPE(FStaticMeshBuilder::BuildMeshVertexPositions);
 
@@ -593,11 +593,12 @@ bool FStaticMeshBuilder::BuildMeshVertexPositions(
 			const FStaticMeshConstAttributes Attributes(MeshDescription);
 			TArrayView<const FVector3f> VertexPositions = Attributes.GetVertexPositions().GetRawArray();
 			TArrayView<const FVertexID> VertexIndices = Attributes.GetTriangleVertexIndices().GetRawArray();
+			const FVector3f BuildScale3D = BuildSettings.BuildScale3D;
 
 			BuiltVertices.Reserve(VertexPositions.Num());
 			for (int32 VertexIndex = 0; VertexIndex < VertexPositions.Num(); ++VertexIndex)
 			{
-				BuiltVertices.Add(VertexPositions[VertexIndex] * BuildSettings.BuildScale3D);
+				BuiltVertices.Add(VertexPositions[VertexIndex] * BuildScale3D);
 			}
 
 			BuiltIndices.Reserve(VertexIndices.Num());
@@ -607,11 +608,11 @@ bool FStaticMeshBuilder::BuildMeshVertexPositions(
 				const uint32 I1 = VertexIndices[TriangleIndex * 3 + 1];
 				const uint32 I2 = VertexIndices[TriangleIndex * 3 + 2];
 
-				const FVector V0 = BuiltVertices[I0];
-				const FVector V1 = BuiltVertices[I1];
-				const FVector V2 = BuiltVertices[I2];
+				const FVector3f V0 = BuiltVertices[I0];
+				const FVector3f V1 = BuiltVertices[I1];
+				const FVector3f V2 = BuiltVertices[I2];
 
-				const FVector TriangleNormal = ((V1 - V2) ^ (V0 - V2));
+				const FVector3f TriangleNormal = ((V1 - V2) ^ (V0 - V2));
 				const bool bDegenerateTriangle = TriangleNormal.SizeSquared() < SMALL_NUMBER;
 				if (!bDegenerateTriangle)
 				{
diff --git a/Engine/Source/Developer/MeshBuilder/Public/IMeshBuilderModule.h b/Engine/Source/Developer/MeshBuilder/Public/IMeshBuilderModule.h
--- a/Engine/Source/Developer/MeshBuilder/Public/IMeshBuilderModule.h
+++ b/Engine/Source/Developer/MeshBuilder/Public/IMeshBuilderModule.h
@@ -31,7 +31,7 @@ public:
 
 	virtual bool BuildMesh(class FStaticMeshRenderData& OutRenderData, class UObject* Mesh, const class FStaticMeshLODGroup& LODGroup) = 0;
 
-	virtual bool BuildMeshVertexPositions(class UObject* StaticMesh, TArray<uint32>& Indices, TArray<FVector>& Vertices) = 0;
+	virtual bool BuildMeshVertexPositions(class UObject* StaticMesh, TArray<uint32>& Indices, TArray<FVector3f>& Vertices) = 0;
 
 	virtual bool BuildSkeletalMesh(const struct FSkeletalMeshBuildParameters& SkeletalMeshBuildParameters) = 0;
 
diff --git a/Engine/Source/Developer/MeshBuilder/Public/MeshBuilder.h b/Engine/Source/Developer/MeshBuilder/Public/MeshBuilder.h
--- a/Engine/Source/Developer/MeshBuilder/Public/MeshBuilder.h
+++ b/Engine/Source/Developer/MeshBuilder/Public/MeshBuilder.h
@@ -29,7 +29,7 @@ public:
 	virtual bool BuildMeshVertexPositions(
 		UStaticMesh* StaticMesh,
 		TArray<uint32>& Indices,
-		TArray<FVector>& Vertices) = 0;
+		TArray<FVector3f>& Vertices) = 0;
 
 	/**
 	 * Build function should be override and is the starting point for skeletal mesh builders
diff --git a/Engine/Source/Developer/MeshBuilder/Public/SkeletalMeshBuilder.h b/Engine/Source/Developer/MeshBuilder/Public/SkeletalMeshBuilder.h
--- a/Engine/Source/Developer/MeshBuilder/Public/SkeletalMeshBuilder.h
+++ b/Engine/Source/Developer/MeshBuilder/Public/SkeletalMeshBuilder.h
@@ -28,7 +28,7 @@ public:
 	virtual bool BuildMeshVertexPositions(
 		UStaticMesh* StaticMesh,
 		TArray<uint32>& Indices,
-		TArray<FVector>& Vertices) override
+		TArray<FVector3f>& Vertices) override
 	{
 		bool No_Support_For_StaticMesh_Build_In_FSkeletalMeshBuilder_Class = false;
 		check(No_Support_For_StaticMesh_Build_In_FSkeletalMeshBuilder_Class);
diff --git a/Engine/Source/Developer/MeshBuilder/Public/StaticMeshBuilder.h b/Engine/Source/Developer/MeshBuilder/Public/StaticMeshBuilder.h
--- a/Engine/Source/Developer/MeshBuilder/Public/StaticMeshBuilder.h
+++ b/Engine/Source/Developer/MeshBuilder/Public/StaticMeshBuilder.h
@@ -31,7 +31,7 @@ public:
 	virtual bool BuildMeshVertexPositions(
 		UStaticMesh* StaticMesh,
 		TArray<uint32>& Indices,
-		TArray<FVector>& Vertices) override;
+		TArray<FVector3f>& Vertices) override;
 
 private:
 
diff --git a/Engine/Source/Developer/MeshUtilities/Private/MeshDistanceFieldUtilities.cpp b/Engine/Source/Developer/MeshUtilities/Private/MeshDistanceFieldUtilities.cpp
--- a/Engine/Source/Developer/MeshUtilities/Private/MeshDistanceFieldUtilities.cpp
+++ b/Engine/Source/Developer/MeshUtilities/Private/MeshDistanceFieldUtilities.cpp
@@ -31,19 +31,19 @@ bool EmbreePointQueryFunction(RTCPointQueryFunctionArguments* args)
 	const int32 TriangleIndex = args->primID;
 	check(TriangleIndex < Context->NumTriangles);
 
-	const FVector* VertexBuffer = (const FVector*)rtcGetGeometryBufferData(Context->MeshGeometry, RTC_BUFFER_TYPE_VERTEX, 0);
+	const FVector3f* VertexBuffer = (const FVector3f*)rtcGetGeometryBufferData(Context->MeshGeometry, RTC_BUFFER_TYPE_VERTEX, 0);
 	const uint32* IndexBuffer = (const uint32*)rtcGetGeometryBufferData(Context->MeshGeometry, RTC_BUFFER_TYPE_INDEX, 0);
 
 	const uint32 I0 = IndexBuffer[TriangleIndex * 3 + 0];
 	const uint32 I1 = IndexBuffer[TriangleIndex * 3 + 1];
 	const uint32 I2 = IndexBuffer[TriangleIndex * 3 + 2];
 
-	const FVector V0 = VertexBuffer[I0];
-	const FVector V1 = VertexBuffer[I1];
-	const FVector V2 = VertexBuffer[I2];
+	const FVector3f V0 = VertexBuffer[I0];
+	const FVector3f V1 = VertexBuffer[I1];
+	const FVector3f V2 = VertexBuffer[I2];
 
-	const FVector QueryPosition(args->query->x, args->query->y, args->query->z);
-	const FVector ClosestPoint = FMath::ClosestPointOnTriangleToPoint(QueryPosition, V0, V1, V2);
+	const FVector3f QueryPosition(args->query->x, args->query->y, args->query->z);
+	const FVector3f ClosestPoint = FMath::ClosestPointOnTriangleToPoint(QueryPosition, V0, V1, V2);
 	const float QueryDistanceSq = (ClosestPoint - QueryPosition).SizeSquared();
 
 	if (QueryDistanceSq < ClosestDistanceSq)
diff --git a/Engine/Source/Developer/MeshUtilities/Private/MeshRepresentationCommon.cpp b/Engine/Source/Developer/MeshUtilities/Private/MeshRepresentationCommon.cpp
--- a/Engine/Source/Developer/MeshUtilities/Private/MeshRepresentationCommon.cpp
+++ b/Engine/Source/Developer/MeshUtilities/Private/MeshRepresentationCommon.cpp
@@ -137,11 +137,11 @@ void MeshRepresentation::SetupEmbreeScene(
 			const uint32 I1 = SourceMeshData.TriangleIndices[TriangleIndex * 3 + 1];
 			const uint32 I2 = SourceMeshData.TriangleIndices[TriangleIndex * 3 + 2];
 
-			const FVector V0 = SourceMeshData.VertexPositions[I0];
-			const FVector V1 = SourceMeshData.VertexPositions[I1];
-			const FVector V2 = SourceMeshData.VertexPositions[I2];
+			const FVector3f V0 = SourceMeshData.VertexPositions[I0];
+			const FVector3f V1 = SourceMeshData.VertexPositions[I1];
+			const FVector3f V2 = SourceMeshData.VertexPositions[I2];
 
-			const FVector TriangleNormal = ((V1 - V2) ^ (V0 - V2));
+			const FVector3f TriangleNormal = ((V1 - V2) ^ (V0 - V2));
 			const bool bDegenerateTriangle = TriangleNormal.SizeSquared() < SMALL_NUMBER;
 			if (!bDegenerateTriangle)
 			{
@@ -158,11 +158,11 @@ void MeshRepresentation::SetupEmbreeScene(
 			const uint32 I1 = Indices[TriangleIndex * 3 + 1];
 			const uint32 I2 = Indices[TriangleIndex * 3 + 2];
 
-			const FVector V0 = LODModel.VertexBuffers.PositionVertexBuffer.VertexPosition(I0);
-			const FVector V1 = LODModel.VertexBuffers.PositionVertexBuffer.VertexPosition(I1);
-			const FVector V2 = LODModel.VertexBuffers.PositionVertexBuffer.VertexPosition(I2);
+			const FVector3f V0 = LODModel.VertexBuffers.PositionVertexBuffer.VertexPosition(I0);
+			const FVector3f V1 = LODModel.VertexBuffers.PositionVertexBuffer.VertexPosition(I1);
+			const FVector3f V2 = LODModel.VertexBuffers.PositionVertexBuffer.VertexPosition(I2);
 
-			const FVector TriangleNormal = ((V1 - V2) ^ (V0 - V2));
+			const FVector3f TriangleNormal = ((V1 - V2) ^ (V0 - V2));
 			const bool bDegenerateTriangle = TriangleNormal.SizeSquared() < SMALL_NUMBER;
 			if (!bDegenerateTriangle)
 			{
@@ -199,14 +199,14 @@ void MeshRepresentation::SetupEmbreeScene(
 	EmbreeScene.Geometry.IndexArray.Empty(NumFilteredIndices);
 	EmbreeScene.Geometry.IndexArray.AddUninitialized(NumFilteredIndices);
 
-	FVector* EmbreeVertices = EmbreeScene.Geometry.VertexArray.GetData();
+	FVector3f* EmbreeVertices = EmbreeScene.Geometry.VertexArray.GetData();
 	uint32* EmbreeIndices = EmbreeScene.Geometry.IndexArray.GetData();
 	EmbreeScene.Geometry.TriangleDescs.Empty(FilteredTriangles.Num());
 
 	for (int32 FilteredTriangleIndex = 0; FilteredTriangleIndex < FilteredTriangles.Num(); FilteredTriangleIndex++)
 	{
 		uint32 I0, I1, I2;
-		FVector V0, V1, V2;
+		FVector3f V0, V1, V2;
 
 		const int32 TriangleIndex = FilteredTriangles[FilteredTriangleIndex];
 		if (SourceMeshData.IsValid())
@@ -280,7 +280,7 @@ void MeshRepresentation::SetupEmbreeScene(
 		RTCGeometry Geometry = rtcNewGeometry(EmbreeScene.EmbreeDevice, RTC_GEOMETRY_TYPE_TRIANGLE);
 		EmbreeScene.Geometry.InternalGeometry = Geometry;
 
-		rtcSetSharedGeometryBuffer(Geometry, RTC_BUFFER_TYPE_VERTEX, 0, RTC_FORMAT_FLOAT3, EmbreeVertices, 0, sizeof(FVector), NumVertices); 
+		rtcSetSharedGeometryBuffer(Geometry, RTC_BUFFER_TYPE_VERTEX, 0, RTC_FORMAT_FLOAT3, EmbreeVertices, 0, sizeof(FVector3f), NumVertices);
 		rtcSetSharedGeometryBuffer(Geometry, RTC_BUFFER_TYPE_INDEX, 0, RTC_FORMAT_UINT3, EmbreeIndices, 0, sizeof(uint32) * 3, FilteredTriangles.Num());
 
 		rtcSetGeometryUserData(Geometry, &EmbreeScene.Geometry);
diff --git a/Engine/Source/Developer/MeshUtilities/Private/MeshRepresentationCommon.h b/Engine/Source/Developer/MeshUtilities/Private/MeshRepresentationCommon.h
--- a/Engine/Source/Developer/MeshUtilities/Private/MeshRepresentationCommon.h
+++ b/Engine/Source/Developer/MeshUtilities/Private/MeshRepresentationCommon.h
@@ -67,7 +67,7 @@ struct FEmbreeTriangleDesc
 struct FEmbreeGeometry
 {
 	TArray<uint32> IndexArray;
-	TArray<FVector> VertexArray;
+	TArray<FVector3f> VertexArray;
 	TArray<FEmbreeTriangleDesc> TriangleDescs; // The material ID of each triangle.
 	RTCGeometry InternalGeometry;
 };
@@ -101,9 +101,9 @@ struct FEmbreeRay : public RTCRayHit
 		hit.primID = RTC_INVALID_GEOMETRY_ID;
 	}
 
-	FVector GetHitNormal() const
+	FVector3f GetHitNormal() const
 	{
-		return FVector(-hit.Ng_x, -hit.Ng_y, -hit.Ng_z).GetSafeNormal();
+		return FVector3f(-hit.Ng_x, -hit.Ng_y, -hit.Ng_z).GetSafeNormal();
 	}	
 
 	bool IsHitTwoSided() const
diff --git a/Engine/Source/Runtime/Engine/Private/DistanceFieldAtlas.cpp b/Engine/Source/Runtime/Engine/Private/DistanceFieldAtlas.cpp
--- a/Engine/Source/Runtime/Engine/Private/DistanceFieldAtlas.cpp
+++ b/Engine/Source/Runtime/Engine/Private/DistanceFieldAtlas.cpp
@@ -122,7 +122,7 @@ FDistanceFieldAsyncQueue* GDistanceFieldAsyncQueue = NULL;
 #if WITH_EDITOR
 
 // DDC key for distance field data, must be changed when modifying the generation code or data format
-#define DISTANCEFIELD_DERIVEDDATA_VER TEXT("CD4A6506-C64C-A229-AA56-2B0A414AE96B")
+#define DISTANCEFIELD_DERIVEDDATA_VER TEXT("1B585A90-19E9-46FB-AF7E-0E84BBC21A7A")
 
 FString BuildDistanceFieldDerivedDataKey(const FString& InMeshKey)
 {
diff --git a/Engine/Source/Runtime/Engine/Private/MeshCardRepresentation.cpp b/Engine/Source/Runtime/Engine/Private/MeshCardRepresentation.cpp
--- a/Engine/Source/Runtime/Engine/Private/MeshCardRepresentation.cpp
+++ b/Engine/Source/Runtime/Engine/Private/MeshCardRepresentation.cpp
@@ -63,7 +63,7 @@ FCardRepresentationAsyncQueue* GCardRepresentationAsyncQueue = NULL;
 #if WITH_EDITOR
 
 // DDC key for card representation data, must be changed when modifying the generation code or data format
-#define CARDREPRESENTATION_DERIVEDDATA_VER TEXT("E8C7CF855B174FE195BE69EAF6000293")
+#define CARDREPRESENTATION_DERIVEDDATA_VER TEXT("0F32177F-E7CA-442E-BEEC-1AEBE4B462BC")
 
 FString BuildCardRepresentationDerivedDataKey(const FString& InMeshKey)
 {
diff --git a/Engine/Source/Runtime/Engine/Public/DerivedMeshDataTaskUtils.h b/Engine/Source/Runtime/Engine/Public/DerivedMeshDataTaskUtils.h
--- a/Engine/Source/Runtime/Engine/Public/DerivedMeshDataTaskUtils.h
+++ b/Engine/Source/Runtime/Engine/Public/DerivedMeshDataTaskUtils.h
@@ -14,7 +14,7 @@ class FSourceMeshDataForDerivedDataTask
 {
 public:
 	TArray<uint32> TriangleIndices;
-	TArray<FVector> VertexPositions;
+	TArray<FVector3f> VertexPositions;
 
 	int32 GetNumIndices() const
 	{
