Workaround embree reading 4 bytes off the end of the shared vertex array.
#b Juan.Canada
#rnx

diff --git a/Engine/Source/Developer/MeshUtilities/Private/MeshRepresentationCommon.cpp b/Engine/Source/Developer/MeshUtilities/Private/MeshRepresentationCommon.cpp
--- a/Engine/Source/Developer/MeshUtilities/Private/MeshRepresentationCommon.cpp
+++ b/Engine/Source/Developer/MeshUtilities/Private/MeshRepresentationCommon.cpp
@@ -191,8 +191,9 @@ void MeshRepresentation::SetupEmbreeScene(
 		}
 	}
 
-	EmbreeScene.Geometry.VertexArray.Empty(NumVertices);
-	EmbreeScene.Geometry.VertexArray.AddUninitialized(NumVertices);
+	const int32 NumBufferVerts = 1; // Reserve extra space at the end of the array, as embree has an internal bug where they read and discard 4 bytes off the end of the array
+	EmbreeScene.Geometry.VertexArray.Empty(NumVertices + NumBufferVerts);
+	EmbreeScene.Geometry.VertexArray.AddUninitialized(NumVertices + NumBufferVerts);
 
 	const int32 NumFilteredIndices = FilteredTriangles.Num() * 3;
 
