Lumen - fixed a bug, where incorrect card direction could be serialized causing an assert to fire at rumtime

diff --git a/Engine/Source/Developer/MeshUtilities/Private/MeshCardRepresentationUtilities.cpp b/Engine/Source/Developer/MeshUtilities/Private/MeshCardRepresentationUtilities.cpp
--- a/Engine/Source/Developer/MeshUtilities/Private/MeshCardRepresentationUtilities.cpp
+++ b/Engine/Source/Developer/MeshUtilities/Private/MeshCardRepresentationUtilities.cpp
@@ -103,14 +103,13 @@ struct FSurfelScene
 class FSurfelCluster
 {
 public:
-	uint8 AxisAlignedDirectionIndex;
-
 	FMatrix44f WorldToLocal;
 	FVector3f Normal;
 	FBox Bounds;
 	TArray<int32> SurfelIndices;
 
 	float MinRayZ = FLT_MAX;
+	uint8 AxisAlignedDirectionIndex = UINT8_MAX;
 
 	// Best surfels to add to this cluster
 	int32 BestSurfelIndex = -1;
@@ -121,6 +120,7 @@ public:
 		Bounds.Init();
 		SurfelIndices.Reset();
 		MinRayZ = FLT_MAX;
+		AxisAlignedDirectionIndex = UINT8_MAX;
 	}
 
 	void SetDirection(uint8 InAxisAlignedDirectionIndex)
@@ -580,6 +580,7 @@ void GrowAllClusters(
 		}
 
 		Cluster.Reset();
+		Cluster.SetDirection(NormalToAxisAlignedDirectionIndex(SurfelScene.Surfels[ClusterSeedIndex].Normal));
 		Cluster.AddSurfel(ClusteringParams, SurfelScene, ClusterSeedIndex);
 		SurfelAssignedToAnyCluster[ClusterSeedIndex] = true;
 	}
@@ -844,6 +845,8 @@ void SerializeLOD(const FGenerateCardMeshContext& Context, const FClusteringPara
 
 	for (const FSurfelCluster& Cluster : LODLevel.Clusters)
 	{
+		check(Cluster.AxisAlignedDirectionIndex < NumAxisAlignedDirections);
+
 		if (Cluster.IsValid(ClusteringParams))
 		{
 			const FMatrix LocalToWorld = Cluster.WorldToLocal.Inverse();
diff --git a/Engine/Source/Runtime/Engine/Private/MeshCardRepresentation.cpp b/Engine/Source/Runtime/Engine/Private/MeshCardRepresentation.cpp
--- a/Engine/Source/Runtime/Engine/Private/MeshCardRepresentation.cpp
+++ b/Engine/Source/Runtime/Engine/Private/MeshCardRepresentation.cpp
@@ -87,7 +87,7 @@ FCardRepresentationAsyncQueue* GCardRepresentationAsyncQueue = NULL;
 #if WITH_EDITOR
 
 // DDC key for card representation data, must be changed when modifying the generation code or data format
-#define CARDREPRESENTATION_DERIVEDDATA_VER TEXT("A1C7CF355B17B3E195BE69EAF6A10292")
+#define CARDREPRESENTATION_DERIVEDDATA_VER TEXT("217ACA355B1713E195BE69EAA6A10290")
 
 FString BuildCardRepresentationDerivedDataKey(const FString& InMeshKey)
 {
