Make sure we use FVector3f in meshutilities
#rnx

diff --git a/Engine/Source/Developer/MeshUtilities/Private/MeshUtilities.cpp b/Engine/Source/Developer/MeshUtilities/Private/MeshUtilities.cpp
--- a/Engine/Source/Developer/MeshUtilities/Private/MeshUtilities.cpp
+++ b/Engine/Source/Developer/MeshUtilities/Private/MeshUtilities.cpp
@@ -153,16 +153,16 @@ void FMeshUtilities::BuildSkeletalAdjacencyIndexBuffer(
 	BuildOptimizationThirdParty::NvTriStripHelper::BuildSkeletalAdjacencyIndexBuffer(VertexBuffer, TexCoordCount, Indices, OutPnAenIndices);
 }
 
-void CalculateTriangleTangentInternal(const FVector& VertexPosA, const FVector2D& VertexUVA
-									  , const FVector& VertexPosB, const FVector2D& VertexUVB
-									  , const FVector& VertexPosC, const FVector2D& VertexUVC
-									  , TArray<FVector>& OutTangents, float CompareThreshold)
+void CalculateTriangleTangentInternal(const FVector3f& VertexPosA, const FVector2D& VertexUVA
+									  , const FVector3f& VertexPosB, const FVector2D& VertexUVB
+									  , const FVector3f& VertexPosC, const FVector2D& VertexUVC
+									  , TArray<FVector3f>& OutTangents, float CompareThreshold)
 {
-	//We must always allocate the OutTangents to 3 FVector
+	//We must always allocate the OutTangents to 3 FVector3f
 	OutTangents.Reset(3);
 	OutTangents.AddZeroed(3);
-	const FVector Positions[3] = { VertexPosA, VertexPosB, VertexPosC };
-	const FVector Normal = ((Positions[1] - Positions[2]) ^ (Positions[0] - Positions[2])).GetSafeNormal(CompareThreshold);
+	const FVector3f Positions[3] = { VertexPosA, VertexPosB, VertexPosC };
+	const FVector3f Normal = ((Positions[1] - Positions[2]) ^ (Positions[0] - Positions[2])).GetSafeNormal(CompareThreshold);
 	//Avoid doing orthonormal vector from a degenerated triangle.
 	if (!Normal.IsNearlyZero(FLT_MIN))
 	{
@@ -187,11 +187,11 @@ void CalculateTriangleTangentInternal(const FVector& VertexPosA, const FVector2D
 		// Use InverseSlow to catch singular matrices.  Inverse can miss this sometimes.
 		const FMatrix TextureToLocal = ParameterToTexture.Inverse() * ParameterToLocal;
 
-		OutTangents[0] = (TextureToLocal.TransformVector(FVector(1, 0, 0)).GetSafeNormal());
-		OutTangents[1] = (TextureToLocal.TransformVector(FVector(0, 1, 0)).GetSafeNormal());
+		OutTangents[0] = (TextureToLocal.TransformVector(FVector3f(1, 0, 0)).GetSafeNormal());
+		OutTangents[1] = (TextureToLocal.TransformVector(FVector3f(0, 1, 0)).GetSafeNormal());
 		OutTangents[2] = (Normal);
 
-		FVector::CreateOrthonormalBasis(
+		FVector3f::CreateOrthonormalBasis(
 			OutTangents[0],
 			OutTangents[1],
 			OutTangents[2]
@@ -200,26 +200,26 @@ void CalculateTriangleTangentInternal(const FVector& VertexPosA, const FVector2D
 		if (OutTangents[0].IsNearlyZero() || OutTangents[0].ContainsNaN()
 			|| OutTangents[1].IsNearlyZero() || OutTangents[1].ContainsNaN())
 		{
-			OutTangents[0] = FVector::ZeroVector;
-			OutTangents[1] = FVector::ZeroVector;
+			OutTangents[0] = FVector3f::ZeroVector;
+			OutTangents[1] = FVector3f::ZeroVector;
 		}
 
 		if (OutTangents[2].IsNearlyZero() || OutTangents[2].ContainsNaN())
 		{
-			OutTangents[2] = FVector::ZeroVector;
+			OutTangents[2] = FVector3f::ZeroVector;
 		}
 	}
 	else
 	{
 		//Add zero tangents and normal for this triangle, this is like weighting it to zero when we compute the vertex normal
 		//But we need the triangle to correctly connect other neighbourg triangles
-		OutTangents[0] = (FVector::ZeroVector);
-		OutTangents[1] = (FVector::ZeroVector);
-		OutTangents[2] = (FVector::ZeroVector);
+		OutTangents[0] = (FVector3f::ZeroVector);
+		OutTangents[1] = (FVector3f::ZeroVector);
+		OutTangents[2] = (FVector3f::ZeroVector);
 	}
 }
 
-void FMeshUtilities::CalculateTriangleTangent(const FSoftSkinVertex& VertexA, const FSoftSkinVertex& VertexB, const FSoftSkinVertex& VertexC, TArray<FVector>& OutTangents, float CompareThreshold)
+void FMeshUtilities::CalculateTriangleTangent(const FSoftSkinVertex& VertexA, const FSoftSkinVertex& VertexB, const FSoftSkinVertex& VertexC, TArray<FVector3f>& OutTangents, float CompareThreshold)
 {
 	CalculateTriangleTangentInternal(VertexA.Position, VertexA.UVs[0]
 									 , VertexB.Position, VertexB.UVs[0]
@@ -1150,10 +1150,10 @@ static void ComputeTriangleTangents(
 	//Currently GetSafeNormal do not support 0.0f threshold properly
 	float RealComparisonThreshold = FMath::Max(ComparisonThreshold, FLT_MIN);
 	
-	TArray<FVector> TriangleTangents;
+	TArray<FVector3f> TriangleTangents;
 	for (int32 TriangleIndex = 0; TriangleIndex < NumTriangles; TriangleIndex++)
 	{
-		FVector Positions[3];
+		FVector3f Positions[3];
 		FVector2D UVs[3];
 		for (int32 Index = 0; Index < 3; ++Index)
 		{
@@ -2254,7 +2254,7 @@ void FMeshUtilities::BuildStaticMeshVertexAndIndexBuffers(
 	const FOverlappingCorners& OverlappingCorners,
 	const TMap<uint32, uint32>& MaterialToSectionMapping,
 	float ComparisonThreshold,
-	FVector BuildScale,
+	FVector3f BuildScale,
 	int32 ImportVersion
 	)
 {
@@ -3341,11 +3341,11 @@ public:
 
 		//Currently GetSafeNormal do not support 0.0f threshold properly
 		float RealComparisonThreshold = FMath::Max(ComparisonThreshold, FLT_MIN);
-		TArray<FVector> TriangleTangents;
+		TArray<FVector3f> TriangleTangents;
 		for (int32 TriangleIndex = 0; TriangleIndex < NumTriangles; TriangleIndex++)
 		{
 			const int32 UVIndex = 0;
-			FVector Positions[3];
+			FVector3f Positions[3];
 			FVector2D UVs[3];
 			for (int32 Index = 0; Index < 3; ++Index)
 			{
diff --git a/Engine/Source/Developer/MeshUtilities/Private/MeshUtilitiesPrivate.h b/Engine/Source/Developer/MeshUtilities/Private/MeshUtilitiesPrivate.h
--- a/Engine/Source/Developer/MeshUtilities/Private/MeshUtilitiesPrivate.h
+++ b/Engine/Source/Developer/MeshUtilities/Private/MeshUtilitiesPrivate.h
@@ -88,7 +88,7 @@ private:
 		const FOverlappingCorners& OverlappingCorners,
 		const TMap<uint32, uint32>& MaterialToSectionMapping,
 		float ComparisonThreshold,
-		FVector BuildScale,
+		FVector3f BuildScale,
 		int32 ImportVersion
 		) override;
 
@@ -149,7 +149,7 @@ private:
 	 *  @param OutTangents - The function allocate the TArray with 3 FVector, to represent the triangle tangent, bi normal and normal.
 	 *  @param CompareThreshold - The threshold use to compare a tangent vector with zero.
 	 */
-	virtual void CalculateTriangleTangent(const FSoftSkinVertex& VertexA, const FSoftSkinVertex& VertexB, const FSoftSkinVertex& VertexC, TArray<FVector>& OutTangents, float CompareThreshold) override;
+	virtual void CalculateTriangleTangent(const FSoftSkinVertex& VertexA, const FSoftSkinVertex& VertexB, const FSoftSkinVertex& VertexC, TArray<FVector3f>& OutTangents, float CompareThreshold) override;
 
 	virtual void CalcBoneVertInfos(USkeletalMesh* SkeletalMesh, TArray<FBoneVertInfo>& Infos, bool bOnlyDominant) override;
 
diff --git a/Engine/Source/Developer/MeshUtilities/Public/MeshUtilities.h b/Engine/Source/Developer/MeshUtilities/Public/MeshUtilities.h
--- a/Engine/Source/Developer/MeshUtilities/Public/MeshUtilities.h
+++ b/Engine/Source/Developer/MeshUtilities/Public/MeshUtilities.h
@@ -194,7 +194,7 @@ public:
 		const FOverlappingCorners& OverlappingCorners,
 		const TMap<uint32, uint32>& MaterialToSectionMapping,
 		float ComparisonThreshold,
-		FVector BuildScale,
+		FVector3f BuildScale,
 		int32 ImportVersion
 		) = 0;
 
@@ -309,7 +309,7 @@ public:
 	 *  @param OutTangents - The function allocate the TArray with 3 FVector, to represent the triangle tangent, bi normal and normal.
 	 *  @param CompareThreshold - The threshold use to compare a tangent vector with zero.
 	 */
-	virtual void CalculateTriangleTangent(const FSoftSkinVertex& VertexA, const FSoftSkinVertex& VertexB, const FSoftSkinVertex& VertexC, TArray<FVector>& OutTangents, float CompareThreshold) = 0;
+	virtual void CalculateTriangleTangent(const FSoftSkinVertex& VertexA, const FSoftSkinVertex& VertexB, const FSoftSkinVertex& VertexC, TArray<FVector3f>& OutTangents, float CompareThreshold) = 0;
 
 	/**
 	 *	Calculate the verts associated weighted to each bone of the skeleton.
diff --git a/Engine/Source/Runtime/Engine/Private/SkeletalMesh.cpp b/Engine/Source/Runtime/Engine/Private/SkeletalMesh.cpp
--- a/Engine/Source/Runtime/Engine/Private/SkeletalMesh.cpp
+++ b/Engine/Source/Runtime/Engine/Private/SkeletalMesh.cpp
@@ -3000,7 +3000,7 @@ void USkeletalMesh::PostLoadVerifyAndFixBadTangent()
 			continue;
 		}
 
-		auto ComputeTriangleTangent = [&MeshUtilities](const FSoftSkinVertex& VertexA, const FSoftSkinVertex& VertexB, const FSoftSkinVertex& VertexC, TArray<FVector>& OutTangents)
+		auto ComputeTriangleTangent = [&MeshUtilities](const FSoftSkinVertex& VertexA, const FSoftSkinVertex& VertexB, const FSoftSkinVertex& VertexC, TArray<FVector3f>& OutTangents)
 		{
 			MeshUtilities.CalculateTriangleTangent(VertexA, VertexB, VertexC, OutTangents, FLT_MIN);
 		};
@@ -3008,7 +3008,7 @@ void USkeletalMesh::PostLoadVerifyAndFixBadTangent()
 		FSkeletalMeshLODModel& ThisLODModel = GetImportedModel()->LODModels[LodIndex];
 		const int32 SectionNum = ThisLODModel.Sections.Num();
 		TArray<FSoftSkinVertex> Vertices;
-		TMap<int32, TArray<FVector>> TriangleTangents;
+		TMap<int32, TArray<FVector3f>> TriangleTangents;
 
 		for (int32 SectionIndex = 0; SectionIndex < SectionNum; ++SectionIndex)
 		{
@@ -3064,13 +3064,13 @@ void USkeletalMesh::PostLoadVerifyAndFixBadTangent()
 						//If the two other axis are valid, fix the tangent with a cross product and normalize the answer.
 						if (TangentB.IsNormalized() && TangentC.IsNormalized())
 						{
-							TangentA = FVector::CrossProduct(TangentB, TangentC);
+							TangentA = FVector3f::CrossProduct(TangentB, TangentC);
 							TangentA.Normalize();
 							return true;
 						}
 
 						//We do not have any valid data to help us for fixing this normal so apply the triangle normals, this will create a faceted mesh but this is better then a black not shade mesh.
-						TArray<FVector>& Tangents = TriangleTangents.FindOrAdd(BaseFaceIndexBufferIndex);
+						TArray<FVector3f>& Tangents = TriangleTangents.FindOrAdd(BaseFaceIndexBufferIndex);
 						if (Tangents.Num() == 0)
 						{
 							const int32 VertexIndex0 = IndexBuffer[BaseFaceIndexBufferIndex] - Section.BaseVertexIndex;
@@ -3085,7 +3085,7 @@ void USkeletalMesh::PostLoadVerifyAndFixBadTangent()
 								return false;
 							}
 							ComputeTriangleTangent(Section.SoftVertices[VertexIndex0], Section.SoftVertices[VertexIndex1], Section.SoftVertices[VertexIndex2], Tangents);
-							const FVector Axis[3] = { {1.0f, 0.0f, 0.0f}, {0.0f, 1.0f, 0.0f}, {0.0f, 0.0f, 1.0f} };
+							const FVector3f Axis[3] = { {1.0f, 0.0f, 0.0f}, {0.0f, 1.0f, 0.0f}, {0.0f, 0.0f, 1.0f} };
 							if (!ensure(Tangents.Num() == 3))
 							{
 								Tangents.Empty(3);
