Fix build after CL# 16355304 merge

diff --git a/Engine/Source/Developer/MeshUtilities/Public/MeshUtilities.h b/Engine/Source/Developer/MeshUtilities/Public/MeshUtilities.h
--- a/Engine/Source/Developer/MeshUtilities/Public/MeshUtilities.h
+++ b/Engine/Source/Developer/MeshUtilities/Public/MeshUtilities.h
@@ -17,7 +17,9 @@ class UStaticMesh;
 class UStaticMeshComponent;
 struct FFlattenMaterial;
 struct FRawMesh;
+struct FRawSkinWeight;
 struct FStaticMeshLODResources;
+class FSkeletalMeshLODModel;
 class FSourceMeshDataForDerivedDataTask;
 
 typedef FIntPoint FMeshIdAndLOD;
