Fixed crash (check firing) in Lumen when importing certain meshes. Check was firigin because of LWC automatic FVector type promotion to double, which caused a mismatch between float and double calculations.

diff --git a/Engine/Source/Developer/MeshUtilities/Private/MeshCardRepresentationUtilities.cpp b/Engine/Source/Developer/MeshUtilities/Private/MeshCardRepresentationUtilities.cpp
--- a/Engine/Source/Developer/MeshUtilities/Private/MeshCardRepresentationUtilities.cpp
+++ b/Engine/Source/Developer/MeshUtilities/Private/MeshCardRepresentationUtilities.cpp
@@ -148,20 +148,19 @@ bool CanAddSurfelToCluster(const FClusteringParams& ClusteringParams, const FSur
 {
 	const FSurfel& SurfelToAdd = SurfelScene.Surfels[SurfelToAddIndex];
 
-	float BoundsMaxZ = Cluster.Bounds.Max.Z;
 	float MinRayZ = Cluster.MinRayZ;
 
-	const float LocalSpacePositonMaxZ = (Cluster.Normal | SurfelToAdd.Position) + ClusteringParams.SurfelExtendRadius;
-	BoundsMaxZ = FMath::Max(BoundsMaxZ, LocalSpacePositonMaxZ);
+	const FVector3f LocalSpacePositon = SurfelToAdd.LocalSurfelPosition[Cluster.AxisAlignedDirectionIndex];
+	float BoundsMaxZ = FMath::Max(Cluster.Bounds.Max.Z, LocalSpacePositon.Z + ClusteringParams.SurfelExtendRadius);
 
 	const float RayTFar = SurfelToAdd.RayCache[Cluster.AxisAlignedDirectionIndex];
 	if (RayTFar < FLT_MAX)
 	{
-		const float ClusterSpaceHitPointZ = LocalSpacePositonMaxZ - ClusteringParams.SurfelExtendRadius + RayTFar;
+		const float ClusterSpaceHitPointZ = LocalSpacePositon.Z + RayTFar;
 		MinRayZ = FMath::Min(MinRayZ, ClusterSpaceHitPointZ);
 	}
 
-	return MinRayZ >= BoundsMaxZ;
+	return MinRayZ > BoundsMaxZ;
 }
 
 float SurfelNormalWeight(const FSurfel& Surfel, FVector3f ClusterNormal, const FClusteringParams& ClusteringParams)
@@ -174,10 +173,10 @@ void FSurfelCluster::AddSurfel(const FClusteringParams& ClusteringParams, const
 	const FSurfel& SurfelToAdd = SurfelScene.Surfels[SurfelToAddIndex];
 	SurfelIndices.Add(SurfelToAddIndex);
 
-	const FVector LocalSpacePositon = SurfelToAdd.LocalSurfelPosition[AxisAlignedDirectionIndex];
+	const FVector3f LocalSpacePositon = SurfelToAdd.LocalSurfelPosition[AxisAlignedDirectionIndex];
 	Bounds += (LocalSpacePositon - ClusteringParams.SurfelExtendRadius);
 	Bounds += (LocalSpacePositon + ClusteringParams.SurfelExtendRadius);
-		
+
 	const float RayTFar = SurfelToAdd.RayCache[AxisAlignedDirectionIndex];
 	if (RayTFar < FLT_MAX)
 	{
@@ -186,7 +185,7 @@ void FSurfelCluster::AddSurfel(const FClusteringParams& ClusteringParams, const
 	}
 
 	// Check if all surfels are visible after add
-	check(MinRayZ >= Bounds.Max.Z);
+	check(MinRayZ > Bounds.Max.Z);
 }
 
 void FSurfelCluster::UpdateBestSurfel(const FClusteringParams& ClusteringParams, const FSurfelScene& SurfelScene, TBitArray<>& SurfelAssignedToAnyCluster)
