Mark up more ParallelFor scopes for Insights

diff --git a/Engine/Source/Developer/CollectionManager/Private/CollectionManager.cpp b/Engine/Source/Developer/CollectionManager/Private/CollectionManager.cpp
--- a/Engine/Source/Developer/CollectionManager/Private/CollectionManager.cpp
+++ b/Engine/Source/Developer/CollectionManager/Private/CollectionManager.cpp
@@ -1835,7 +1835,8 @@ void FCollectionManager::LoadCollections()
 
 	FRWLock CollectionLock;
 	ParallelFor(
-		ECollectionShareType::CST_All,
+		TEXT("LoadCollections.PF"),
+		ECollectionShareType::CST_All,1,
 		[this, &CollectionLock](int32 CacheIdx)
 		{
 			const ECollectionShareType::Type ShareType = ECollectionShareType::Type(CacheIdx);
@@ -1847,7 +1848,8 @@ void FCollectionManager::LoadCollections()
 			IFileManager::Get().FindFiles(Filenames, *WildCard, true, false);
 
 			ParallelFor(
-				Filenames.Num(),
+				TEXT("LoadCollections.PF"),
+				Filenames.Num(),1,
 				[this, &CollectionLock, &Filenames, &CollectionFolder, bUseSCC, ShareType](int32 FilenameIdx)
 				{
 					const FString& BaseFilename = Filenames[FilenameIdx];
diff --git a/Engine/Source/Developer/MeshUtilities/Private/MeshCardRepresentationUtilities.cpp b/Engine/Source/Developer/MeshUtilities/Private/MeshCardRepresentationUtilities.cpp
--- a/Engine/Source/Developer/MeshUtilities/Private/MeshCardRepresentationUtilities.cpp
+++ b/Engine/Source/Developer/MeshUtilities/Private/MeshCardRepresentationUtilities.cpp
@@ -550,6 +550,8 @@ void GenerateSurfels(
 	FSurfelScene& SurfelScene,
 	FClusteringParams& ClusteringParams)
 {
+	TRACE_CPUPROFILER_EVENT_SCOPE(GenerateSurfels);
+
 	const uint32 NumSourceVertices = Context.EmbreeScene.Geometry.VertexArray.Num();
 	const uint32 NumSourceIndices = Context.EmbreeScene.Geometry.IndexArray.Num();
 	const int32 NumSourceTriangles = NumSourceIndices / 3;
@@ -645,7 +647,7 @@ void GenerateSurfels(
 	const int32 DebugSurfelDirection = MeshCardRepresentation::GetDebugSurfelDirection();
 
 	const bool bSingleThreaded = CVarCardRepresentationParallelBuild.GetValueOnAnyThread() == 0;
-	ParallelFor(MeshCardGen::NumAxisAlignedDirections,
+	ParallelFor( TEXT("GenerateSurfels.PF"), MeshCardGen::NumAxisAlignedDirections,1,
 		[&](int32 AxisAlignedDirectionIndex)
 		{
 			if (DebugSurfelDirection < 0 || DebugSurfelDirection == AxisAlignedDirectionIndex)
@@ -660,7 +662,7 @@ void GenerateSurfels(
 					SurfelScenePerDirection
 				);
 			}			
-		}, bSingleThreaded);
+		}, bSingleThreaded? EParallelForFlags::ForceSingleThread : EParallelForFlags::None);
 
 	ClusteringParams.VoxelSize = VoxelSize;
 	ClusteringParams.MaxSurfelDistanceXY = MeshCardRepresentation::GetMaxSurfelDistanceXY();
@@ -1038,7 +1040,7 @@ void BuildMeshCardsLOD1(const FBox& MeshBounds, const FGenerateCardMeshContext&
 
 	// Generate initial list of clusters
 	const bool bSingleThreaded = CVarCardRepresentationParallelBuild.GetValueOnAnyThread() == 0;
-	ParallelFor(MeshCardGen::NumAxisAlignedDirections,
+	ParallelFor( TEXT("BuildMeshCardsLOD1.PF"), MeshCardGen::NumAxisAlignedDirections,1,
 		[&](int32 AxisAlignedDirectionIndex)
 		{
 			const FSurfelScenePerDirection& SurfelScenePerDirection = SurfelScene.Directions[AxisAlignedDirectionIndex];
@@ -1088,13 +1090,13 @@ void BuildMeshCardsLOD1(const FBox& MeshBounds, const FGenerateCardMeshContext&
 					break;
 				}
 			}
-		}, bSingleThreaded);
+		}, bSingleThreaded? EParallelForFlags::ForceSingleThread : EParallelForFlags::None);
 
 	LimitClusters(ClusteringParams, SurfelScene, LODLevel);
 
 	if (LODLevel.NeedToReGrow())
 	{
-		ParallelFor(MeshCardGen::NumAxisAlignedDirections,
+		ParallelFor( TEXT("BuildMeshCardsLOD1.PF"), MeshCardGen::NumAxisAlignedDirections,1,
 			[&](int32 AxisAlignedDirectionIndex)
 			{
 				if (LODLevel.Directions[AxisAlignedDirectionIndex].bNeedsToReGrow)
@@ -1131,7 +1133,7 @@ void BuildMeshCardsLOD1(const FBox& MeshBounds, const FGenerateCardMeshContext&
 						}
 					}
 				}
-			}, bSingleThreaded);
+			}, bSingleThreaded? EParallelForFlags::ForceSingleThread : EParallelForFlags::None);
 	}
 
 	UpdateLODLevelCoverage(SurfelScene, ClusteringParams, LODLevel);
diff --git a/Engine/Source/Runtime/CoreUObject/Private/UObject/GarbageCollection.cpp b/Engine/Source/Runtime/CoreUObject/Private/UObject/GarbageCollection.cpp
--- a/Engine/Source/Runtime/CoreUObject/Private/UObject/GarbageCollection.cpp
+++ b/Engine/Source/Runtime/CoreUObject/Private/UObject/GarbageCollection.cpp
@@ -1324,7 +1324,7 @@ public:
 
 		// Iterate over all objects. Note that we iterate over the UObjectArray and usually check only internal flags which
 		// are part of the array so we don't suffer from cache misses as much as we would if we were to check ObjectFlags.
-		ParallelFor(NumThreads, [ObjectsToSerializeArrays, &ClustersToDissolveList, &KeepClusterRefsList, FastKeepFlags, KeepFlags, NumberOfObjectsPerThread, NumThreads, MaxNumberOfObjects](int32 ThreadIndex)
+		ParallelFor( TEXT("GarbageCollection.PF"),NumThreads,1, [ObjectsToSerializeArrays, &ClustersToDissolveList, &KeepClusterRefsList, FastKeepFlags, KeepFlags, NumberOfObjectsPerThread, NumThreads, MaxNumberOfObjects](int32 ThreadIndex)
 		{
 			TRACE_CPUPROFILER_EVENT_SCOPE(MarkObjectsAsUnreachableTask);
 			int32 FirstObjectIndex = ThreadIndex * NumberOfObjectsPerThread + GUObjectArray.GetFirstGCIndex();
@@ -1436,7 +1436,7 @@ public:
 			}
 
 			GObjectCountDuringLastMarkPhase.Add(ObjectCountDuringMarkPhase);
-		}, !bParallel);
+		}, !bParallel ? EParallelForFlags::ForceSingleThread : EParallelForFlags::None);
 		
 		// Collect all objects to serialize from all threads and put them into a single array
 		{
@@ -2109,7 +2109,7 @@ void GatherUnreachableObjects(bool bForceSingleThreaded)
 
 	// Iterate over all objects. Note that we iterate over the UObjectArray and usually check only internal flags which
 	// are part of the array so we don't suffer from cache misses as much as we would if we were to check ObjectFlags.
-	ParallelFor(NumThreads, [&ClusterItemsToDestroy, NumberOfObjectsPerThread, NumThreads, MaxNumberOfObjects](int32 ThreadIndex)
+	ParallelFor( TEXT("GarbageCollection.PF"),NumThreads,1, [&ClusterItemsToDestroy, NumberOfObjectsPerThread, NumThreads, MaxNumberOfObjects](int32 ThreadIndex)
 	{
 		TRACE_CPUPROFILER_EVENT_SCOPE(GatherUnreachableObjectsTask);
 		int32 FirstObjectIndex = ThreadIndex * NumberOfObjectsPerThread + (GExitPurge ? 0 : GUObjectArray.GetFirstGCIndex());
@@ -2138,7 +2138,7 @@ void GatherUnreachableObjects(bool bForceSingleThreaded)
 			GUnreachableObjects.Append(ThisThreadUnreachableObjects);
 			ClusterItemsToDestroy.Append(ThisThreadClusterItemsToDestroy);
 		}
-	}, bForceSingleThreaded);
+	}, bForceSingleThreaded ? EParallelForFlags::ForceSingleThread : EParallelForFlags::None);
 
 	{
 		// @todo: if GUObjectClusters.FreeCluster() was thread safe we could do this in parallel too
diff --git a/Engine/Source/Runtime/CoreUObject/Private/UObject/GarbageCollectionVerification.cpp b/Engine/Source/Runtime/CoreUObject/Private/UObject/GarbageCollectionVerification.cpp
--- a/Engine/Source/Runtime/CoreUObject/Private/UObject/GarbageCollectionVerification.cpp
+++ b/Engine/Source/Runtime/CoreUObject/Private/UObject/GarbageCollectionVerification.cpp
@@ -114,7 +114,7 @@ void VerifyGCAssumptions()
 	int32 NumberOfObjectsPerThread = (MaxNumberOfObjects / NumThreads) + 1;
 	FGCArrayStruct* ArrayStructs = new FGCArrayStruct[NumThreads];
 
-	ParallelFor(NumThreads, [&ReferenceCollector, ArrayStructs, NumberOfObjectsPerThread, NumThreads, MaxNumberOfObjects](int32 ThreadIndex)
+	ParallelFor( TEXT("GarbageCollection.PF"),NumThreads,1, [&ReferenceCollector, ArrayStructs, NumberOfObjectsPerThread, NumThreads, MaxNumberOfObjects](int32 ThreadIndex)
 	{
 		int32 FirstObjectIndex = ThreadIndex * NumberOfObjectsPerThread;
 		int32 NumObjects = (ThreadIndex < (NumThreads - 1)) ? NumberOfObjectsPerThread : (MaxNumberOfObjects - (NumThreads - 1)*NumberOfObjectsPerThread);
@@ -293,7 +293,7 @@ void VerifyClustersAssumptions()
 	FGCArrayStruct* ArrayStructs = new FGCArrayStruct[NumThreads];
 	FThreadSafeCounter NumErrors(0);
 
-	ParallelFor(NumThreads, [&NumErrors, ArrayStructs, NumberOfClustersPerThread, NumThreads, MaxNumberOfClusters](int32 ThreadIndex)
+	ParallelFor( TEXT("GarbageCollection.PF"),NumThreads,1, [&NumErrors, ArrayStructs, NumberOfClustersPerThread, NumThreads, MaxNumberOfClusters](int32 ThreadIndex)
 	{
 		int32 FirstClusterIndex = ThreadIndex * NumberOfClustersPerThread;
 		int32 NumClusters = (ThreadIndex < (NumThreads - 1)) ? NumberOfClustersPerThread : (MaxNumberOfClusters - (NumThreads - 1) * NumberOfClustersPerThread);
@@ -345,7 +345,7 @@ void VerifyObjectFlagMirroring()
 	int32 NumberOfObjectsPerThread = (MaxNumberOfObjects / NumThreads) + 1;
 	std::atomic<uint32> NumErrors(0);
 
-	ParallelFor(NumThreads, [&NumErrors, NumberOfObjectsPerThread, NumThreads, MaxNumberOfObjects](int32 ThreadIndex)
+	ParallelFor( TEXT("GarbageCollection.PF"),NumThreads,1, [&NumErrors, NumberOfObjectsPerThread, NumThreads, MaxNumberOfObjects](int32 ThreadIndex)
 	{
 		int32 FirstObjectIndex = ThreadIndex * NumberOfObjectsPerThread;
 		int32 NumObjects = (ThreadIndex < (NumThreads - 1)) ? NumberOfObjectsPerThread : (MaxNumberOfObjects - (NumThreads - 1) * NumberOfObjectsPerThread);
diff --git a/Engine/Source/Runtime/Engine/Private/ShaderCompiler/ShaderCompiler.cpp b/Engine/Source/Runtime/Engine/Private/ShaderCompiler/ShaderCompiler.cpp
--- a/Engine/Source/Runtime/Engine/Private/ShaderCompiler/ShaderCompiler.cpp
+++ b/Engine/Source/Runtime/Engine/Private/ShaderCompiler/ShaderCompiler.cpp
@@ -7373,6 +7373,8 @@ static inline FShader* ProcessCompiledJob(FShaderCompileJob* SingleJob, const FS
 
 void ProcessCompiledGlobalShaders(const TArray<FShaderCommonCompileJobPtr>& CompilationResults)
 {
+	TRACE_CPUPROFILER_EVENT_SCOPE(ProcessCompiledGlobalShaders);
+
 	UE_LOG(LogShaders, Verbose, TEXT("Compiled %u global shaders"), CompilationResults.Num());
 
 	TArray<EShaderPlatform> ShaderPlatformsProcessed;
diff --git a/Engine/Source/Runtime/Experimental/Chaos/Public/Chaos/ChaosMarshallingManager.h b/Engine/Source/Runtime/Experimental/Chaos/Public/Chaos/ChaosMarshallingManager.h
--- a/Engine/Source/Runtime/Experimental/Chaos/Public/Chaos/ChaosMarshallingManager.h
+++ b/Engine/Source/Runtime/Experimental/Chaos/Public/Chaos/ChaosMarshallingManager.h
@@ -139,7 +139,7 @@ public:
 	template <typename Lambda>
 	void ParallelForEachProxy(const Lambda& Func)
 	{
-		::ParallelFor(DirtyProxyBucketInfo.TotalNum,[this,&Func](int32 Idx)
+		::ParallelFor( TEXT("Chaos.PF"),DirtyProxyBucketInfo.TotalNum,1, [this,&Func](int32 Idx)
 		{
 			int32 BucketIdx, InnerIdx;
 			DirtyProxyBucketInfo.GetBucketIdx(Idx, BucketIdx, InnerIdx);
diff --git a/Engine/Source/Runtime/StaticMeshDescription/Private/StaticMeshOperations.cpp b/Engine/Source/Runtime/StaticMeshDescription/Private/StaticMeshOperations.cpp
--- a/Engine/Source/Runtime/StaticMeshDescription/Private/StaticMeshOperations.cpp
+++ b/Engine/Source/Runtime/StaticMeshDescription/Private/StaticMeshOperations.cpp
@@ -259,7 +259,7 @@ void FStaticMeshOperations::ComputeTriangleTangentsAndNormals(FMeshDescription&
 	const int32 BatchSize = 8 * 1024;
 	const int32 BatchCount = (NumTriangles + BatchSize - 1) / BatchSize;
 
-	ParallelFor(BatchCount,
+	ParallelFor( TEXT("ComputeTriangleTangentsAndNormals.PF"), BatchCount,1,
 		[BatchSize, ComparisonThreshold, NumTriangles, &Attributes](int32 BatchIndex)
 		{
 			TArrayView<const FVector3f> VertexPositions = Attributes.GetVertexPositions().GetRawArray();
@@ -1376,7 +1376,7 @@ void FStaticMeshOperations::ComputeTangentsAndNormals(FMeshDescription& MeshDesc
 	const int32 BatchCount = (NumVertices + BatchSize - 1) / BatchSize;
 
 	//Iterate all vertex to compute normals for all vertex instance
-	ParallelFor(BatchCount,
+	ParallelFor( TEXT("ComputeTangentsAndNormals.PF"), BatchCount,1,
 		[NumVertices, BatchSize, bComputeTangentWithMikkTSpace, bComputeWeightedNormals, &MeshDescription, &TriangleDatas](int32 BatchIndex)
 		{
 			FStaticMeshAttributes Attributes(MeshDescription);
