Lumen - fixed surface coverage on sparse trees
* Mostly two sided meshes are assumed to be foliage and forced to be repsented as a cube projection
* Lowered card min density treshold for cube projections

diff --git a/Engine/Source/Developer/MeshUtilities/Private/MeshCardRepresentationUtilities.cpp b/Engine/Source/Developer/MeshUtilities/Private/MeshCardRepresentationUtilities.cpp
--- a/Engine/Source/Developer/MeshUtilities/Private/MeshCardRepresentationUtilities.cpp
+++ b/Engine/Source/Developer/MeshUtilities/Private/MeshCardRepresentationUtilities.cpp
@@ -609,7 +609,7 @@ void GenerateSurfelsForDirection(
 	}
 }
 
-void InitClusteringParams(FClusteringParams& ClusteringParams, const FBox& MeshCardsBounds, int32 MaxVoxels, int32 MaxLumenMeshCards, float MinClusterCoverage)
+void InitClusteringParams(FClusteringParams& ClusteringParams, const FBox& MeshCardsBounds, int32 MaxVoxels, int32 MaxLumenMeshCards, int32 LODLevel)
 {
 	const float TargetVoxelSize = 20.0f;
 
@@ -689,8 +689,8 @@ void InitClusteringParams(FClusteringParams& ClusteringParams, const FBox& MeshC
 
 	ClusteringParams.VoxelSize = VoxelSize;
 	ClusteringParams.MaxSurfelDistanceXY = MeshCardRepresentation::GetMaxSurfelDistanceXY();
-	ClusteringParams.MinClusterCoverage = MinClusterCoverage;
-	ClusteringParams.MinDensityPerCluster = MeshCardRepresentation::GetMinDensity();
+	ClusteringParams.MinClusterCoverage = LODLevel == 0 ? 0.5f : 10.0f;
+	ClusteringParams.MinDensityPerCluster = MeshCardRepresentation::GetMinDensity() * (LODLevel == 0 ? 0.1f : 1.0f);
 	ClusteringParams.MaxLumenMeshCards = MaxLumenMeshCards;
 	ClusteringParams.bDebug = MeshCardRepresentation::IsDebugMode();
 	ClusteringParams.bSingleThreadedBuild = CVarCardRepresentationParallelBuild.GetValueOnAnyThread() == 0;
@@ -700,8 +700,6 @@ void InitSurfelScene(
 	const FGenerateCardMeshContext& Context,
 	const FBox& MeshCardsBounds,
 	int32 MaxLumenMeshCards,
-	float MinClusterCoverageLOD0,
-	float MinClusterCoverageLOD1,
 	FSurfelScene& SurfelScene,
 	FClusteringParams& ClusteringParamsLOD0,
 	FClusteringParams& ClusteringParamsLOD1)
@@ -733,8 +731,8 @@ void InitSurfelScene(
 
 	do
 	{
-		InitClusteringParams(ClusteringParamsLOD0, MeshCardsBounds, MaxVoxels, MaxLumenMeshCards, MinClusterCoverageLOD0);
-		InitClusteringParams(ClusteringParamsLOD1, MeshCardsBounds, MaxVoxels, MaxLumenMeshCards, MinClusterCoverageLOD1);
+		InitClusteringParams(ClusteringParamsLOD0, MeshCardsBounds, MaxVoxels, MaxLumenMeshCards, 0);
+		InitClusteringParams(ClusteringParamsLOD1, MeshCardsBounds, MaxVoxels, MaxLumenMeshCards, 1);
 
 		ParallelFor(TEXT("InitSurfelScene.PF"), MeshCardGen::NumAxisAlignedDirections, 1,
 			[&](int32 AxisAlignedDirectionIndex)
@@ -978,7 +976,7 @@ void BuildMeshCardsLOD0(const FBox& MeshBounds, const FGenerateCardMeshContext&
 			LODLevel.Directions[AxisAlignedDirectionIndex].Clusters.Add(TempCluster);
 		}
 	}
-
+	
 	UpdateLODLevelCoverage(SurfelScene, ClusteringParams, LODLevel);
 }
 
@@ -1369,13 +1367,17 @@ void BuildMeshCards(const FBox& MeshBounds, const FGenerateCardMeshContext& Cont
 	FSurfelScene SurfelScene;
 	FClusteringParams ClusteringParamsLOD0;
 	FClusteringParams ClusteringParamsLOD1;
-	InitSurfelScene(Context, MeshCardsBounds, MaxLumenMeshCards, /*MinClusterCoverageLOD0*/ 0.5f, /*MinClusterCoverageLOD1*/ 10.0f, SurfelScene, ClusteringParamsLOD0, ClusteringParamsLOD1);
+	InitSurfelScene(Context, MeshCardsBounds, MaxLumenMeshCards, SurfelScene, ClusteringParamsLOD0, ClusteringParamsLOD1);
 
 	FMeshCardsLODLevel MeshCardsLOD0;
 	BuildMeshCardsLOD0(MeshBounds, Context, SurfelScene, ClusteringParamsLOD0, MeshCardsLOD0);
 
+	// Assume that two sided is foliage and revert to a simpler and more reliable box projection
 	FMeshCardsLODLevel MeshCardsLOD1;
-	BuildMeshCardsLOD1(MeshBounds, Context, SurfelScene, ClusteringParamsLOD1, MeshCardsLOD1);
+	if (!Context.EmbreeScene.bMostlyTwoSided)
+	{
+		BuildMeshCardsLOD1(MeshBounds, Context, SurfelScene, ClusteringParamsLOD1, MeshCardsLOD1);
+	}
 
 	OutData.MeshCardsBuildData.Bounds = MeshCardsBounds;
 	OutData.MeshCardsBuildData.MaxLODLevel = 0;
diff --git a/Engine/Source/Developer/MeshUtilities/Private/MeshDistanceFieldUtilities.cpp b/Engine/Source/Developer/MeshUtilities/Private/MeshDistanceFieldUtilities.cpp
--- a/Engine/Source/Developer/MeshUtilities/Private/MeshDistanceFieldUtilities.cpp
+++ b/Engine/Source/Developer/MeshUtilities/Private/MeshDistanceFieldUtilities.cpp
@@ -274,29 +274,6 @@ void FMeshUtilities::GenerateSignedDistanceFieldVolumeData(
 
 		check(EmbreeScene.bUseEmbree);
 
-		bool bMostlyTwoSided;
-		{
-			uint32 NumTrianglesTotal = 0;
-			uint32 NumTwoSidedTriangles = 0;
-
-			for (int32 SectionIndex = 0; SectionIndex < LODModel.Sections.Num(); SectionIndex++)
-			{
-				const FStaticMeshSection& Section = LODModel.Sections[SectionIndex];
-
-				if (MaterialBlendModes.IsValidIndex(Section.MaterialIndex))
-				{
-					NumTrianglesTotal += Section.NumTriangles;
-
-					if (MaterialBlendModes[Section.MaterialIndex].bTwoSided)
-					{
-						NumTwoSidedTriangles += Section.NumTriangles;
-					}
-				}
-			}
-
-			bMostlyTwoSided = NumTwoSidedTriangles * 4 >= NumTrianglesTotal || bGenerateAsIfTwoSided;
-		}
-
 		// Whether to use an Embree Point Query to compute the closest unsigned distance.  Rays will only be traced to determine backfaces visible for sign.
 		const bool bUsePointQuery = true;
 
@@ -339,7 +316,7 @@ void FMeshUtilities::GenerateSignedDistanceFieldVolumeData(
 		// We sample on voxel corners and use central differencing for gradients, so a box mesh using two-sided materials whose vertices lie on LocalSpaceMeshBounds produces a zero gradient on intersection
 		// Expand the mesh bounds by a fraction of a voxel to allow room for a pullback on the hit location for computing the gradient.
 		// Only expand for two sided meshes as this adds significant Mesh SDF tracing cost
-		if (bMostlyTwoSided)
+		if (EmbreeScene.bMostlyTwoSided)
 		{
 			const FVector DesiredDimensions = FVector(LocalSpaceMeshBounds.GetSize() * FVector(NumVoxelsPerLocalSpaceUnit / (float)DistanceField::UniqueDataBrickSize));
 			const FIntVector Mip0IndirectionDimensions = FIntVector(
@@ -512,7 +489,7 @@ void FMeshUtilities::GenerateSignedDistanceFieldVolumeData(
 
 		MeshRepresentation::DeleteEmbreeScene(EmbreeScene);
 
-		OutData.bMostlyTwoSided = bMostlyTwoSided;
+		OutData.bMostlyTwoSided = EmbreeScene.bMostlyTwoSided;
 		OutData.LocalSpaceMeshBounds = LocalSpaceMeshBounds;
 
 		OutData.StreamableMips.Lock(LOCK_READ_WRITE);
diff --git a/Engine/Source/Developer/MeshUtilities/Private/MeshRepresentationCommon.cpp b/Engine/Source/Developer/MeshUtilities/Private/MeshRepresentationCommon.cpp
--- a/Engine/Source/Developer/MeshUtilities/Private/MeshRepresentationCommon.cpp
+++ b/Engine/Source/Developer/MeshUtilities/Private/MeshRepresentationCommon.cpp
@@ -305,6 +305,29 @@ void MeshRepresentation::SetupEmbreeScene(
 	{
 		EmbreeScene.kDopTree.Build(BuildTriangles);
 	}
+
+	// bMostlyTwoSided
+	{
+		uint32 NumTrianglesTotal = 0;
+		uint32 NumTwoSidedTriangles = 0;
+
+		for (int32 SectionIndex = 0; SectionIndex < LODModel.Sections.Num(); SectionIndex++)
+		{
+			const FStaticMeshSection& Section = LODModel.Sections[SectionIndex];
+
+			if (MaterialBlendModes.IsValidIndex(Section.MaterialIndex))
+			{
+				NumTrianglesTotal += Section.NumTriangles;
+
+				if (MaterialBlendModes[Section.MaterialIndex].bTwoSided)
+				{
+					NumTwoSidedTriangles += Section.NumTriangles;
+				}
+			}
+		}
+
+		EmbreeScene.bMostlyTwoSided = NumTwoSidedTriangles * 4 >= NumTrianglesTotal || bGenerateAsIfTwoSided;
+	}
 }
 
 void MeshRepresentation::DeleteEmbreeScene(FEmbreeScene& EmbreeScene)
diff --git a/Engine/Source/Developer/MeshUtilities/Private/MeshRepresentationCommon.h b/Engine/Source/Developer/MeshUtilities/Private/MeshRepresentationCommon.h
--- a/Engine/Source/Developer/MeshUtilities/Private/MeshRepresentationCommon.h
+++ b/Engine/Source/Developer/MeshUtilities/Private/MeshRepresentationCommon.h
@@ -83,6 +83,7 @@ class FEmbreeScene
 public:
 	bool bUseEmbree = false;
 	int32 NumIndices = 0;
+	bool bMostlyTwoSided = false;
 
 	// Embree
 	RTCDevice EmbreeDevice = nullptr;
diff --git a/Engine/Source/Runtime/Engine/Private/MeshCardRepresentation.cpp b/Engine/Source/Runtime/Engine/Private/MeshCardRepresentation.cpp
--- a/Engine/Source/Runtime/Engine/Private/MeshCardRepresentation.cpp
+++ b/Engine/Source/Runtime/Engine/Private/MeshCardRepresentation.cpp
@@ -138,7 +138,7 @@ FCardRepresentationAsyncQueue* GCardRepresentationAsyncQueue = NULL;
 #if WITH_EDITOR
 
 // DDC key for card representation data, must be changed when modifying the generation code or data format
-#define CARDREPRESENTATION_DERIVEDDATA_VER TEXT("69E68026-A68B-A032-A3DF-E7901648DB72")
+#define CARDREPRESENTATION_DERIVEDDATA_VER TEXT("66C4BDF4-D9C5-4B26-8556-A3F9FCB2A2EA")
 
 FString BuildCardRepresentationDerivedDataKey(const FString& InMeshKey, int32 MaxLumenMeshCards)
 {
