Lumen - disable mesh cards generation when max num cards in static mesh properties is set to 0

diff --git a/Engine/Source/Developer/MeshUtilities/Private/MeshCardRepresentationUtilities.cpp b/Engine/Source/Developer/MeshUtilities/Private/MeshCardRepresentationUtilities.cpp
--- a/Engine/Source/Developer/MeshUtilities/Private/MeshCardRepresentationUtilities.cpp
+++ b/Engine/Source/Developer/MeshUtilities/Private/MeshCardRepresentationUtilities.cpp
@@ -1415,37 +1415,41 @@ bool FMeshUtilities::GenerateCardRepresentationData(
 {
 #if USE_EMBREE
 	TRACE_CPUPROFILER_EVENT_SCOPE(FMeshUtilities::GenerateCardRepresentationData);
-	const double StartTime = FPlatformTime::Seconds();
 
-	FEmbreeScene EmbreeScene;
-	MeshRepresentation::SetupEmbreeScene(MeshName,
-		SourceMeshData,
-		LODModel,
-		MaterialBlendModes,
-		bGenerateAsIfTwoSided,
-		EmbreeScene);
-
-	if (!EmbreeScene.EmbreeScene)
+	if (MaxLumenMeshCards > 0)
 	{
-		return false;
-	}
+		const double StartTime = FPlatformTime::Seconds();
+
+		FEmbreeScene EmbreeScene;
+		MeshRepresentation::SetupEmbreeScene(MeshName,
+			SourceMeshData,
+			LODModel,
+			MaterialBlendModes,
+			bGenerateAsIfTwoSided,
+			EmbreeScene);
+
+		if (!EmbreeScene.EmbreeScene)
+		{
+			return false;
+		}
 
-	FGenerateCardMeshContext Context(MeshName, EmbreeScene, OutData);
+		FGenerateCardMeshContext Context(MeshName, EmbreeScene, OutData);
 
-	// Note: must operate on the SDF bounds when available, because SDF generation can expand the mesh's bounds
-	const FBox BuildCardsBounds = DistanceFieldVolumeData && DistanceFieldVolumeData->LocalSpaceMeshBounds.IsValid ? DistanceFieldVolumeData->LocalSpaceMeshBounds : Bounds.GetBox();
-	BuildMeshCards(BuildCardsBounds, Context, MaxLumenMeshCards, OutData);
+		// Note: must operate on the SDF bounds when available, because SDF generation can expand the mesh's bounds
+		const FBox BuildCardsBounds = DistanceFieldVolumeData && DistanceFieldVolumeData->LocalSpaceMeshBounds.IsValid ? DistanceFieldVolumeData->LocalSpaceMeshBounds : Bounds.GetBox();
+		BuildMeshCards(BuildCardsBounds, Context, MaxLumenMeshCards, OutData);
 
-	MeshRepresentation::DeleteEmbreeScene(EmbreeScene);
+		MeshRepresentation::DeleteEmbreeScene(EmbreeScene);
 
-	const float TimeElapsed = (float)(FPlatformTime::Seconds() - StartTime);
-	if (TimeElapsed > 1.0f)
-	{
-		UE_LOG(LogMeshUtilities, Log, TEXT("Finished mesh card build in %.1fs %s tris:%d surfels:%d"),
-			TimeElapsed,
-			*MeshName,
-			EmbreeScene.NumIndices / 3,
-			OutData.MeshCardsBuildData.DebugData.NumSurfels);
+		const float TimeElapsed = (float)(FPlatformTime::Seconds() - StartTime);
+		if (TimeElapsed > 1.0f)
+		{
+			UE_LOG(LogMeshUtilities, Log, TEXT("Finished mesh card build in %.1fs %s tris:%d surfels:%d"),
+				TimeElapsed,
+				*MeshName,
+				EmbreeScene.NumIndices / 3,
+				OutData.MeshCardsBuildData.DebugData.NumSurfels);
+		}
 	}
 
 	return true;
