Updated MeshUtilities::GenerateStratifiedUniformHemisphereSamples to use concentric disk sampling.
- The number of samples returned is the square number closer to the requested number.
- Updated code paths using the function to request a number of samples close to what they previously got.

diff --git a/Engine/Source/Developer/MeshUtilities/Private/MeshCardRepresentationUtilities.cpp b/Engine/Source/Developer/MeshUtilities/Private/MeshCardRepresentationUtilities.cpp
--- a/Engine/Source/Developer/MeshUtilities/Private/MeshCardRepresentationUtilities.cpp
+++ b/Engine/Source/Developer/MeshUtilities/Private/MeshCardRepresentationUtilities.cpp
@@ -716,7 +716,7 @@
 	}
 
 	// Generate random ray directions over a hemisphere
-	constexpr uint32 NumRayDirectionsOverHemisphere = 64;
+	constexpr uint32 NumRayDirectionsOverHemisphere = 36;
 	TArray<FVector3f> RayDirectionsOverHemisphere;
 	{
 		FRandomStream RandomStream(0);
diff --git a/Engine/Source/Developer/MeshUtilities/Private/MeshDistanceFieldUtilities.cpp b/Engine/Source/Developer/MeshUtilities/Private/MeshDistanceFieldUtilities.cpp
--- a/Engine/Source/Developer/MeshUtilities/Private/MeshDistanceFieldUtilities.cpp
+++ b/Engine/Source/Developer/MeshUtilities/Private/MeshDistanceFieldUtilities.cpp
@@ -279,7 +279,7 @@
 
 		TArray<FVector3f> SampleDirections;
 		{
-			const int32 NumVoxelDistanceSamples = bUsePointQuery ? 120 : 1200;
+			const int32 NumVoxelDistanceSamples = bUsePointQuery ? 49 : 576;
 			FRandomStream RandomStream(0);
 			MeshUtilities::GenerateStratifiedUniformHemisphereSamples(NumVoxelDistanceSamples, RandomStream, SampleDirections);
 			TArray<FVector3f> OtherHemisphereSamples;
diff --git a/Engine/Source/Developer/MeshUtilities/Private/MeshRepresentationCommon.cpp b/Engine/Source/Developer/MeshUtilities/Private/MeshRepresentationCommon.cpp
--- a/Engine/Source/Developer/MeshUtilities/Private/MeshRepresentationCommon.cpp
+++ b/Engine/Source/Developer/MeshUtilities/Private/MeshRepresentationCommon.cpp
@@ -5,27 +5,55 @@
 #include "MeshUtilitiesPrivate.h"
 #include "DerivedMeshDataTaskUtils.h"
 
+static FVector3f UniformSampleHemisphere(FVector2D Uniforms)
+{
+	Uniforms = Uniforms * 2.0f - 1.0f;
+
+	if(Uniforms == FVector2D::ZeroVector)
+	{
+		return FVector3f::ZeroVector;
+	}
+
+	float R;
+	float Theta;
+
+	if (FMath::Abs(Uniforms.X) > FMath::Abs(Uniforms.Y))
+	{
+		R = Uniforms.X;
+		Theta = (float)PI / 4 * (Uniforms.Y / Uniforms.X);
+	}
+	else
+	{
+		R = Uniforms.Y;
+		Theta = (float)PI / 2 - (float)PI / 4 * (Uniforms.X / Uniforms.Y);
+	}
+
+	// concentric disk sample
+	const float U = R * FMath::Cos(Theta);
+	const float V = R * FMath::Sin(Theta);
+	const float R2 = R * R;
+
+	// map to hemisphere [P. Shirley, Kenneth Chiu; 1997; A Low Distortion Map Between Disk and Square]
+	return FVector3f(U * FMath::Sqrt(2 - R2), V * FMath::Sqrt(2 - R2), 1.0f - R2);
+}
+
 void MeshUtilities::GenerateStratifiedUniformHemisphereSamples(int32 NumSamples, FRandomStream& RandomStream, TArray<FVector3f>& Samples)
 {
-	const int32 NumThetaSteps = FMath::TruncToInt(FMath::Sqrt(NumSamples / (2.0f * (float)PI)));
-	const int32 NumPhiSteps = FMath::TruncToInt(NumThetaSteps * (float)PI);
+	const int32 NumSamplesDim = FMath::TruncToInt(FMath::Sqrt((float)NumSamples));
 
-	Samples.Empty(NumThetaSteps * NumPhiSteps);
-	for (int32 ThetaIndex = 0; ThetaIndex < NumThetaSteps; ThetaIndex++)
+	Samples.Empty(NumSamplesDim * NumSamplesDim);
+
+	for (int32 IndexX = 0; IndexX < NumSamplesDim; IndexX++)
 	{
-		for (int32 PhiIndex = 0; PhiIndex < NumPhiSteps; PhiIndex++)
+		for (int32 IndexY = 0; IndexY < NumSamplesDim; IndexY++)
 		{
 			const float U1 = RandomStream.GetFraction();
 			const float U2 = RandomStream.GetFraction();
 
-			const float Fraction1 = (ThetaIndex + U1) / (float)NumThetaSteps;
-			const float Fraction2 = (PhiIndex + U2) / (float)NumPhiSteps;
+			const float Fraction1 = (IndexX + U1) / (float)NumSamplesDim;
+			const float Fraction2 = (IndexY + U2) / (float)NumSamplesDim;
 
-			const float R = FMath::Sqrt(1.0f - Fraction1 * Fraction1);
-
-			const float Phi = 2.0f * (float)PI * Fraction2;
-			// Convert to Cartesian
-			Samples.Add(FVector3f(FMath::Cos(Phi) * R, FMath::Sin(Phi) * R, Fraction1));
+			Samples.Add(UniformSampleHemisphere(FVector2D(Fraction1, Fraction2)));
 		}
 	}
 }
diff --git a/Engine/Source/Runtime/Engine/Private/DistanceFieldAtlas.cpp b/Engine/Source/Runtime/Engine/Private/DistanceFieldAtlas.cpp
--- a/Engine/Source/Runtime/Engine/Private/DistanceFieldAtlas.cpp
+++ b/Engine/Source/Runtime/Engine/Private/DistanceFieldAtlas.cpp
@@ -127,7 +127,7 @@
 #if WITH_EDITOR
 
 // DDC key for distance field data, must be changed when modifying the generation code or data format
-#define DISTANCEFIELD_DERIVEDDATA_VER TEXT("540160fd-f0b5-41da-bbe3-f573428ccd0d")
+#define DISTANCEFIELD_DERIVEDDATA_VER TEXT("CB37D8F6-5ED7-476E-BDFB-A0319D4099DC")
 
 FString BuildDistanceFieldDerivedDataKey(const FString& InMeshKey)
 {
diff --git a/Engine/Source/Runtime/Engine/Private/MeshCardRepresentation.cpp b/Engine/Source/Runtime/Engine/Private/MeshCardRepresentation.cpp
--- a/Engine/Source/Runtime/Engine/Private/MeshCardRepresentation.cpp
+++ b/Engine/Source/Runtime/Engine/Private/MeshCardRepresentation.cpp
@@ -138,7 +138,7 @@
 #if WITH_EDITOR
 
 // DDC key for card representation data, must be changed when modifying the generation code or data format
-#define CARDREPRESENTATION_DERIVEDDATA_VER TEXT("66C4BDF4-D9C5-4B26-8556-A3F9FCB2A2EA")
+#define CARDREPRESENTATION_DERIVEDDATA_VER TEXT("3F7067D7-0552-4C66-BBC5-F2589BDB798A")
 
 FString BuildCardRepresentationDerivedDataKey(const FString& InMeshKey, int32 MaxLumenMeshCards)
 {
