Lumen - Surface Cache Improvements

* Card sampling with squared weights instead of cutoff in order to allow surface to leak into tiny areas without coverage
* Bump min static mesh LOD to reduce mismatches between surface cache and geometry

Card placement:
* Move near plane back outside of mesh bounds during capture in order to handle displacement mismatches
* Discard surfels seeing triangle back faces only when sky isn't visible
* New additive card placement, which can be used for continuous LOD. Instead of clustering, it tries to find the best split plane until the entire surface isn't covered.
* Lower voxel and min card size to be able to cover smaller areas and place more cards
* Spawn multiple surfels per cell if they cannot be merged
* Fixed surfel generation and card size

diff --git a/Engine/Source/Developer/MeshUtilities/Private/MeshCardRepresentationUtilities.cpp b/Engine/Source/Developer/MeshUtilities/Private/MeshCardRepresentationUtilities.cpp
--- a/Engine/Source/Developer/MeshUtilities/Private/MeshCardRepresentationUtilities.cpp
+++ b/Engine/Source/Developer/MeshUtilities/Private/MeshCardRepresentationUtilities.cpp
@@ -315,7 +315,9 @@ FSurfelVisibility ComputeSurfelVisibility(
 		SurfelSampleIndex = (SurfelSampleIndex + 1) % SurfelSamplesNum;
 	}
 
-	const bool bInsideGeometry = NumHits > 0 && NumBackFaceHits > RayDirectionsOverHemisphere.Num() * 0.2f;
+	const bool bInsideGeometry =
+		NumHits > 0.8f * RayDirectionsOverHemisphere.Num()
+		&& NumBackFaceHits > 0.2f * RayDirectionsOverHemisphere.Num();
 
 	FSurfelVisibility SurfelVisibility;
 	SurfelVisibility.Visibility = VisibilitySum / RayDirectionsOverHemisphere.Num();
