Lumen Surface Cache - add a small card margin to handle displacement, which can move mesh vertices outside of the original bounding box.

diff --git a/Engine/Source/Developer/MeshUtilities/Private/MeshCardRepresentationUtilities.cpp b/Engine/Source/Developer/MeshUtilities/Private/MeshCardRepresentationUtilities.cpp
--- a/Engine/Source/Developer/MeshUtilities/Private/MeshCardRepresentationUtilities.cpp
+++ b/Engine/Source/Developer/MeshUtilities/Private/MeshCardRepresentationUtilities.cpp
@@ -924,11 +924,15 @@ void SerializeLOD(
 			FVector3f ClusterBoundsMin = (FVector3f(Cluster.Bounds.Min) - FVector3f(0.0f, 0.0f, 0.5f)) * ClusteringParams.VoxelSize;
 			FVector3f ClusterBoundsMax = (FVector3f(Cluster.Bounds.Max) + FVector3f(1.0f, 1.0f, 1.5f)) * ClusteringParams.VoxelSize;
 
-			// Clamp sides to mesh bounds, but leave the near plane, as LOD/displacement may move it outside of bounds
+			// Clamp to mesh bounds
+			// Leave small margin for Z as LOD/displacement may move it outside of bounds
+			static float MarginZ = 10.0f;
 			ClusterBoundsMin.X = FMath::Max(ClusterBoundsMin.X, LocalMeshCardsBounds.Min.X);
 			ClusterBoundsMin.Y = FMath::Max(ClusterBoundsMin.Y, LocalMeshCardsBounds.Min.Y);
+			ClusterBoundsMin.Z = FMath::Min(ClusterBoundsMax.Z, LocalMeshCardsBounds.Min.Z - MarginZ);
 			ClusterBoundsMax.X = FMath::Min(ClusterBoundsMax.X, LocalMeshCardsBounds.Max.X);
 			ClusterBoundsMax.Y = FMath::Min(ClusterBoundsMax.Y, LocalMeshCardsBounds.Max.Y);
+			ClusterBoundsMax.Z = FMath::Min(ClusterBoundsMax.Z, LocalMeshCardsBounds.Max.Z + MarginZ);
 
 			const FVector3f ClusterBoundsOrigin = (ClusterBoundsMax + ClusterBoundsMin) * 0.5f;
 			const FVector3f ClusterBoundsExtent = (ClusterBoundsMax - ClusterBoundsMin) * 0.5f;
@@ -1010,7 +1014,7 @@ void BuildMeshCards(const FBox& MeshBounds, const FGenerateCardMeshContext& Cont
 {
 	// Make sure BBox isn't empty and we can generate card representation for it. This handles e.g. infinitely thin planes.
 	const FVector MeshCardsBoundsCenter = MeshBounds.GetCenter();
-	const FVector MeshCardsBoundsExtent = FVector::Max(MeshBounds.GetExtent() + 1.0f, FVector(5.0f));
+	const FVector MeshCardsBoundsExtent = FVector::Max(MeshBounds.GetExtent() + 1.0f, FVector(1.0f));
 	const FBox MeshCardsBounds(MeshCardsBoundsCenter - MeshCardsBoundsExtent, MeshCardsBoundsCenter + MeshCardsBoundsExtent);
 
 	// Prepare a list of surfels for cluster fitting
diff --git a/Engine/Source/Runtime/Engine/Private/MeshCardRepresentation.cpp b/Engine/Source/Runtime/Engine/Private/MeshCardRepresentation.cpp
--- a/Engine/Source/Runtime/Engine/Private/MeshCardRepresentation.cpp
+++ b/Engine/Source/Runtime/Engine/Private/MeshCardRepresentation.cpp
@@ -105,7 +105,7 @@ FCardRepresentationAsyncQueue* GCardRepresentationAsyncQueue = NULL;
 #if WITH_EDITOR
 
 // DDC key for card representation data, must be changed when modifying the generation code or data format
-#define CARDREPRESENTATION_DERIVEDDATA_VER TEXT("A3695875-02C2-159E-8923-2F7B6816A281")
+#define CARDREPRESENTATION_DERIVEDDATA_VER TEXT("EBA88EB9-EA65-46C3-B742-71EFA5CB333E")
 
 FString BuildCardRepresentationDerivedDataKey(const FString& InMeshKey, int32 MaxLumenMeshCards)
 {
