Lumen Mesh Cards - fixed bug forcing all cards to be placed on the bounding box of a mesh.

diff --git a/Engine/Source/Developer/MeshUtilities/Private/MeshCardRepresentationUtilities.cpp b/Engine/Source/Developer/MeshUtilities/Private/MeshCardRepresentationUtilities.cpp
--- a/Engine/Source/Developer/MeshUtilities/Private/MeshCardRepresentationUtilities.cpp
+++ b/Engine/Source/Developer/MeshUtilities/Private/MeshCardRepresentationUtilities.cpp
@@ -929,7 +929,7 @@ void SerializeLOD(
 			static float MarginZ = 10.0f;
 			ClusterBoundsMin.X = FMath::Max(ClusterBoundsMin.X, LocalMeshCardsBounds.Min.X);
 			ClusterBoundsMin.Y = FMath::Max(ClusterBoundsMin.Y, LocalMeshCardsBounds.Min.Y);
-			ClusterBoundsMin.Z = FMath::Min(ClusterBoundsMax.Z, LocalMeshCardsBounds.Min.Z - MarginZ);
+			ClusterBoundsMin.Z = FMath::Max(ClusterBoundsMin.Z, LocalMeshCardsBounds.Min.Z - MarginZ);
 			ClusterBoundsMax.X = FMath::Min(ClusterBoundsMax.X, LocalMeshCardsBounds.Max.X);
 			ClusterBoundsMax.Y = FMath::Min(ClusterBoundsMax.Y, LocalMeshCardsBounds.Max.Y);
 			ClusterBoundsMax.Z = FMath::Min(ClusterBoundsMax.Z, LocalMeshCardsBounds.Max.Z + MarginZ);
diff --git a/Engine/Source/Runtime/Engine/Private/MeshCardRepresentation.cpp b/Engine/Source/Runtime/Engine/Private/MeshCardRepresentation.cpp
--- a/Engine/Source/Runtime/Engine/Private/MeshCardRepresentation.cpp
+++ b/Engine/Source/Runtime/Engine/Private/MeshCardRepresentation.cpp
@@ -105,7 +105,7 @@ FCardRepresentationAsyncQueue* GCardRepresentationAsyncQueue = NULL;
 #if WITH_EDITOR
 
 // DDC key for card representation data, must be changed when modifying the generation code or data format
-#define CARDREPRESENTATION_DERIVEDDATA_VER TEXT("EBA88EB9-EA65-46C3-B742-71EFA5CB333E")
+#define CARDREPRESENTATION_DERIVEDDATA_VER TEXT("18971EB0-9015-4184-8B29-D3C2B8CEF111")
 
 FString BuildCardRepresentationDerivedDataKey(const FString& InMeshKey, int32 MaxLumenMeshCards)
 {
