Lumen Mesh Cards - fixed infinite loop when generating mesh cards for large meshes and EmbreeRay.ray.tfar + 0.01f was equal to EmbreeRay.ray.tfar. Now hits against previous primitive are ignored and ray is always advanced by a smallest possible increment.

diff --git a/Engine/Source/Developer/MeshUtilities/Private/MeshCardRepresentationUtilities.cpp b/Engine/Source/Developer/MeshUtilities/Private/MeshCardRepresentationUtilities.cpp
--- a/Engine/Source/Developer/MeshUtilities/Private/MeshCardRepresentationUtilities.cpp
+++ b/Engine/Source/Developer/MeshUtilities/Private/MeshCardRepresentationUtilities.cpp
@@ -374,6 +374,7 @@ void GenerateSurfelsForDirection(
 
 				// Cell index where any geometry was last found
 				int32 LastHitCoordZ = -2;
+				int32 SkipPrimId = RTC_INVALID_GEOMETRY_ID;
 				float RayTNear = 0.0f;
 
 				while (LastHitCoordZ + 1 < ClusterBasis.VolumeSize.Z)
@@ -390,6 +391,7 @@ void GenerateSurfelsForDirection(
 
 					FEmbreeIntersectionContext EmbreeContext;
 					rtcInitIntersectContext(&EmbreeContext);
+					EmbreeContext.SkipPrimId = SkipPrimId;
 					rtcIntersect1(Context.EmbreeScene.EmbreeScene, &EmbreeContext, &EmbreeRay);
 
 					if (EmbreeRay.hit.geomID != RTC_INVALID_GEOMETRY_ID && EmbreeRay.hit.primID != RTC_INVALID_GEOMETRY_ID)
@@ -421,8 +423,10 @@ void GenerateSurfelsForDirection(
 							}
 						}
 
+						// Move ray to the next intersection
 						LastHitCoordZ = HitCoordZ;
-						RayTNear = FMath::Max(NearPlaneOffset + (LastHitCoordZ + 1) * ClusteringParams.VoxelSize, EmbreeRay.ray.tfar + 0.01f);
+						RayTNear = std::nextafter(FMath::Max(NearPlaneOffset + (LastHitCoordZ + 1) * ClusteringParams.VoxelSize, EmbreeRay.ray.tfar), std::numeric_limits<float>::infinity());
+						SkipPrimId = EmbreeRay.hit.primID;
 					}
 					else
 					{
diff --git a/Engine/Source/Developer/MeshUtilities/Private/MeshRepresentationCommon.cpp b/Engine/Source/Developer/MeshUtilities/Private/MeshRepresentationCommon.cpp
--- a/Engine/Source/Developer/MeshUtilities/Private/MeshRepresentationCommon.cpp
+++ b/Engine/Source/Developer/MeshUtilities/Private/MeshRepresentationCommon.cpp
@@ -93,6 +93,13 @@ void EmbreeFilterFunc(const struct RTCFilterFunctionNArguments* args)
 
 	FEmbreeIntersectionContext& IntersectionContext = *static_cast<FEmbreeIntersectionContext*>(args->context);
 	IntersectionContext.ElementIndex = Desc.ElementIndex;
+
+	const RTCHit& EmbreeHit = *(RTCHit*)args->hit;
+	if (IntersectionContext.SkipPrimId != RTC_INVALID_GEOMETRY_ID && IntersectionContext.SkipPrimId == EmbreeHit.primID)
+	{
+		// Ignore hit in order to continue tracing
+		args->valid[0] = 0;
+	}
 }
 
 void EmbreeErrorFunc(void* userPtr, RTCError code, const char* str)
diff --git a/Engine/Source/Developer/MeshUtilities/Private/MeshRepresentationCommon.h b/Engine/Source/Developer/MeshUtilities/Private/MeshRepresentationCommon.h
--- a/Engine/Source/Developer/MeshUtilities/Private/MeshRepresentationCommon.h
+++ b/Engine/Source/Developer/MeshUtilities/Private/MeshRepresentationCommon.h
@@ -135,6 +135,9 @@ struct FEmbreeIntersectionContext : public RTCIntersectContext
 		return ElementIndex == 1;
 	}
 
+	// Hit against this primitive will be ignored
+	int32 SkipPrimId = RTC_INVALID_GEOMETRY_ID;
+
 	// Additional Outputs.
 	int32 ElementIndex; // Material Index
 };
