Remove doubles from mesh processing

diff --git a/Engine/Plugins/Runtime/GeometryCache/Source/GeometryCache/Private/CodecV1Test.h b/Engine/Plugins/Runtime/GeometryCache/Source/GeometryCache/Private/CodecV1Test.h
--- a/Engine/Plugins/Runtime/GeometryCache/Source/GeometryCache/Private/CodecV1Test.h
+++ b/Engine/Plugins/Runtime/GeometryCache/Source/GeometryCache/Private/CodecV1Test.h
@@ -93,7 +93,7 @@ public:
 					const FVector3f& PositionA = OriginalMeshData.Positions[Index];
 					const FVector3f& PositionB = DecodedMeshData.Positions[Index];
 
-					if (!PointsEqual((FVector)PositionA, (FVector)PositionB, true))
+					if (!PointsEqual(PositionA, PositionB, true))
 					{
 						check(true);
 					}
@@ -133,7 +133,7 @@ public:
 					// Motion vectors if we have any
 					if (OriginalMeshData.Positions.Num() == OriginalMeshData.MotionVectors.Num())
 					{
-						if (!PointsEqual((FVector)OriginalMeshData.MotionVectors[Index], (FVector)DecodedMeshData.MotionVectors[Index]))
+						if (!PointsEqual(OriginalMeshData.MotionVectors[Index], DecodedMeshData.MotionVectors[Index]))
 						{
 							check(true);
 						}
diff --git a/Engine/Plugins/Runtime/GeometryCache/Source/GeometryCache/Private/GeometryCachePreprocessor.cpp b/Engine/Plugins/Runtime/GeometryCache/Source/GeometryCache/Private/GeometryCachePreprocessor.cpp
--- a/Engine/Plugins/Runtime/GeometryCache/Source/GeometryCache/Private/GeometryCachePreprocessor.cpp
+++ b/Engine/Plugins/Runtime/GeometryCache/Source/GeometryCache/Private/GeometryCachePreprocessor.cpp
@@ -133,7 +133,7 @@ bool FOptimizeGeometryCachePreprocessor::AreIndexedVerticesEqual(int32 IndexBuff
 		// Motion vectors if we have any
 		if (BufferedFrames[Frame].MeshData.Positions.Num() == BufferedFrames[Frame].MeshData.MotionVectors.Num())
 		{
-			if (!PointsEqual((FVector)BufferedFrames[Frame].MeshData.MotionVectors[VertexIndexA], (FVector)BufferedFrames[Frame].MeshData.MotionVectors[VertexIndexB]))
+			if (!PointsEqual(BufferedFrames[Frame].MeshData.MotionVectors[VertexIndexA], BufferedFrames[Frame].MeshData.MotionVectors[VertexIndexB]))
 			{
 				return false;
 			}
diff --git a/Engine/Source/Developer/MeshBuilder/Private/StaticMeshBuilder.cpp b/Engine/Source/Developer/MeshBuilder/Private/StaticMeshBuilder.cpp
--- a/Engine/Source/Developer/MeshBuilder/Private/StaticMeshBuilder.cpp
+++ b/Engine/Source/Developer/MeshBuilder/Private/StaticMeshBuilder.cpp
@@ -703,9 +703,9 @@ bool FStaticMeshBuilder::BuildMeshVertexPositions(
 bool AreVerticesEqual(FStaticMeshBuildVertex const& A, FStaticMeshBuildVertex const& B, float ComparisonThreshold)
 {
 	if (   !A.Position.Equals(B.Position, ComparisonThreshold)
-		|| !NormalsEqual((FVector)A.TangentX, (FVector)B.TangentX)
-		|| !NormalsEqual((FVector)A.TangentY, (FVector)B.TangentY)
-		|| !NormalsEqual((FVector)A.TangentZ, (FVector)B.TangentZ)
+		|| !NormalsEqual(A.TangentX, B.TangentX)
+		|| !NormalsEqual(A.TangentY, B.TangentY)
+		|| !NormalsEqual(A.TangentZ, B.TangentZ)
 		|| A.Color != B.Color)
 	{
 		return false;
@@ -759,7 +759,7 @@ void BuildVertexBuffer(
 	const bool bHasColors = VertexInstanceColors.IsValid();
 
 	const uint32 NumTextureCoord = VertexInstanceUVs.IsValid() ? VertexInstanceUVs.GetNumChannels() : 0;
-	const FMatrix ScaleMatrix = FScaleMatrix(BuildSettings.BuildScale3D).Inverse().GetTransposed();
+	const FVector3f BuildScale(BuildSettings.BuildScale3D);
 
 	TMap<FPolygonGroupID, int32> PolygonGroupToSectionIndex;
 
@@ -799,10 +799,10 @@ void BuildVertexBuffer(
 
 		TArrayView<const FVertexID> VertexIDs = MeshDescription.GetTriangleVertices(TriangleID);
 
-		FVector CornerPositions[3];
+		FVector3f CornerPositions[3];
 		for (int32 TriVert = 0; TriVert < 3; ++TriVert)
 		{
-			CornerPositions[TriVert] = (FVector)VertexPositions[VertexIDs[TriVert]];
+			CornerPositions[TriVert] = VertexPositions[VertexIDs[TriVert]];
 		}
 		FOverlappingThresholds OverlappingThresholds;
 		OverlappingThresholds.ThresholdPosition = VertexComparisonThreshold;
@@ -819,25 +819,28 @@ void BuildVertexBuffer(
 		for (int32 TriVert = 0; TriVert < 3; ++TriVert, ++WedgeIndex)
 		{
 			const FVertexInstanceID VertexInstanceID = VertexInstanceIDs[TriVert];
-			const FVector& VertexPosition = CornerPositions[TriVert];
-			const FVector& VertexInstanceNormal = (FVector)VertexInstanceNormals[VertexInstanceID];
-			const FVector& VertexInstanceTangent = (FVector)VertexInstanceTangents[VertexInstanceID];
+			const FVector3f& VertexPosition = CornerPositions[TriVert];
+			const FVector3f& VertexInstanceNormal = VertexInstanceNormals[VertexInstanceID];
+			const FVector3f& VertexInstanceTangent = VertexInstanceTangents[VertexInstanceID];
 			const float VertexInstanceBinormalSign = VertexInstanceBinormalSigns[VertexInstanceID];
 
 			FStaticMeshBuildVertex StaticMeshVertex;
 
-			StaticMeshVertex.Position = FVector3f(VertexPosition * BuildSettings.BuildScale3D);
+			StaticMeshVertex.Position = VertexPosition * BuildScale;
 			if (bNeedTangents)
 			{
-				StaticMeshVertex.TangentX = (FVector4f)ScaleMatrix.TransformVector(VertexInstanceTangent).GetSafeNormal();
-				StaticMeshVertex.TangentY = (FVector4f)ScaleMatrix.TransformVector(FVector::CrossProduct(VertexInstanceNormal, VertexInstanceTangent) * VertexInstanceBinormalSign).GetSafeNormal();
+				StaticMeshVertex.TangentX = VertexInstanceTangent / BuildScale;
+				StaticMeshVertex.TangentY = ( (VertexInstanceNormal ^ VertexInstanceTangent) * VertexInstanceBinormalSign ) / BuildScale;
+				StaticMeshVertex.TangentX.Normalize();
+				StaticMeshVertex.TangentY.Normalize();
 			}
 			else
 			{
 				StaticMeshVertex.TangentX = FVector3f(1.0f, 0.0f, 0.0f);
 				StaticMeshVertex.TangentY = FVector3f(0.0f, 1.0f, 0.0f);
 			}
-			StaticMeshVertex.TangentZ = (FVector4f)ScaleMatrix.TransformVector(VertexInstanceNormal).GetSafeNormal();
+			StaticMeshVertex.TangentZ = VertexInstanceNormal / BuildScale;
+			StaticMeshVertex.TangentZ.Normalize();
 
 			if (bHasColors)
 			{
diff --git a/Engine/Source/Developer/MeshSimplifier/Private/MeshTessellate.h b/Engine/Source/Developer/MeshSimplifier/Private/MeshTessellate.h
--- a/Engine/Source/Developer/MeshSimplifier/Private/MeshTessellate.h
+++ b/Engine/Source/Developer/MeshSimplifier/Private/MeshTessellate.h
@@ -5,9 +5,6 @@
 #include "CoreMinimal.h"
 #include "Containers/HashTable.h"
 
-template< typename T > FORCEINLINE uint32 Min3Index( const T A, const T B, const T C ) { return ( A < B ) ? ( ( A < C ) ? 0 : 2 ) : ( ( B < C ) ? 1 : 2 ); }
-template< typename T > FORCEINLINE uint32 Max3Index( const T A, const T B, const T C ) { return ( A > B ) ? ( ( A > C ) ? 0 : 2 ) : ( ( B > C ) ? 1 : 2 ); }
-
 FORCEINLINE uint32 Hash3( const uint32 x, const uint32 y, const uint32 z )
 {
 	//return ( 73856093 * x ) ^ ( 15485867 * y ) ^ ( 83492791 * z );
@@ -82,7 +79,7 @@ void MeshTessellate( TArray<T>& Verts, TArray< uint32 >& Indexes, float MaxEdgeS
 			EdgeLength2[i] = ( p0 - p1 ).SizeSquared();
 		}
 
-		const uint32 e0 = Max3Index( EdgeLength2[0], EdgeLength2[1], EdgeLength2[2] );
+		const uint32 e0 = FMath::Max3Index( EdgeLength2[0], EdgeLength2[1], EdgeLength2[2] );
 		const uint32 e1 = (1 << e0) & 3;
 		const uint32 e2 = (1 << e1) & 3;
 		if( EdgeLength2[ e0 ] < MaxEdgeSize * MaxEdgeSize )
diff --git a/Engine/Source/Developer/MeshSimplifier/Private/QuadricMeshReduction.cpp b/Engine/Source/Developer/MeshSimplifier/Private/QuadricMeshReduction.cpp
--- a/Engine/Source/Developer/MeshSimplifier/Private/QuadricMeshReduction.cpp
+++ b/Engine/Source/Developer/MeshSimplifier/Private/QuadricMeshReduction.cpp
@@ -126,9 +126,9 @@ public:
 			}
 
 			// Don't process degenerate triangles.
-			if( PointsEqual((FVector)CornerPositions[0], (FVector)CornerPositions[1]) ||
-				PointsEqual((FVector)CornerPositions[0], (FVector)CornerPositions[2]) ||
-				PointsEqual((FVector)CornerPositions[1], (FVector)CornerPositions[2]) )
+			if( PointsEqual(CornerPositions[0], CornerPositions[1]) ||
+				PointsEqual(CornerPositions[0], CornerPositions[2]) ||
+				PointsEqual(CornerPositions[1], CornerPositions[2]) )
 			{
 				WedgeIndex += 3;
 				continue;
diff --git a/Engine/Source/Developer/MeshSimplifier/Private/SimpVert.h b/Engine/Source/Developer/MeshSimplifier/Private/SimpVert.h
--- a/Engine/Source/Developer/MeshSimplifier/Private/SimpVert.h
+++ b/Engine/Source/Developer/MeshSimplifier/Private/SimpVert.h
@@ -34,10 +34,10 @@ public:
 
 	bool		Equals(	const VertType& a ) const
 	{
-		if( !PointsEqual( (FVector)Position, (FVector)a.Position ) ||
-			!NormalsEqual( (FVector)Tangents[0], (FVector)a.Tangents[0] ) ||
-			!NormalsEqual( (FVector)Tangents[1], (FVector)a.Tangents[1] ) ||
-			!NormalsEqual( (FVector)Normal, (FVector)a.Normal ) ||
+		if( !PointsEqual( Position, a.Position ) ||
+			!NormalsEqual( Tangents[0], a.Tangents[0] ) ||
+			!NormalsEqual( Tangents[1], a.Tangents[1] ) ||
+			!NormalsEqual( Normal, a.Normal ) ||
 			!Color.Equals( a.Color ) )
 		{
 			return false;
diff --git a/Engine/Source/Developer/MeshUtilities/Private/MeshUtilities.cpp b/Engine/Source/Developer/MeshUtilities/Private/MeshUtilities.cpp
--- a/Engine/Source/Developer/MeshUtilities/Private/MeshUtilities.cpp
+++ b/Engine/Source/Developer/MeshUtilities/Private/MeshUtilities.cpp
@@ -57,6 +57,7 @@
 #include "Algo/Transform.h"
 #include "Rendering/SkeletalMeshModel.h"
 #include "Rendering/SkeletalMeshRenderData.h"
+#include "Math/Bounds.h"
 
 #include "LandscapeProxy.h"
 #include "Landscape.h"
@@ -2128,7 +2129,7 @@ static void BuildDepthOnlyIndexBuffer(
 	VertIndexAndZ.Empty(NumVertices);
 	for (int32 VertIndex = 0; VertIndex < NumVertices; VertIndex++)
 	{
-		new(VertIndexAndZ)FIndexAndZ(VertIndex, (FVector)InVertices[VertIndex].Position);
+		new(VertIndexAndZ)FIndexAndZ(VertIndex, InVertices[VertIndex].Position);
 	}
 	VertIndexAndZ.Sort(FCompareIndexAndZ());
 
@@ -2222,9 +2223,9 @@ static bool AreVerticesEqual(
 	)
 {
 	if (!PointsEqual(A.Position, B.Position, ComparisonThreshold)
-		|| !NormalsEqual((FVector)A.TangentX, (FVector)B.TangentX)
-		|| !NormalsEqual((FVector)A.TangentY, (FVector)B.TangentY)
-		|| !NormalsEqual((FVector)A.TangentZ, (FVector)B.TangentZ)
+		|| !NormalsEqual(A.TangentX, B.TangentX)
+		|| !NormalsEqual(A.TangentY, B.TangentY)
+		|| !NormalsEqual(A.TangentZ, B.TangentZ)
 		|| A.Color != B.Color)
 	{
 		return false;
@@ -2905,23 +2906,25 @@ public:
 		}
 
 		// Calculate the bounding box.
-		FBox BoundingBox(ForceInit);
+		FBounds3f Bounds;
 		FPositionVertexBuffer& BasePositionVertexBuffer = OutRenderData.LODResources[0].VertexBuffers.PositionVertexBuffer;
 		for (uint32 VertexIndex = 0; VertexIndex < BasePositionVertexBuffer.GetNumVertices(); VertexIndex++)
 		{
-			BoundingBox += (FVector)BasePositionVertexBuffer.VertexPosition(VertexIndex);
+			Bounds += BasePositionVertexBuffer.VertexPosition(VertexIndex);
 		}
-		BoundingBox.GetCenterAndExtents(OutRenderData.Bounds.Origin, OutRenderData.Bounds.BoxExtent);
+		OutRenderData.Bounds.Origin = (FVector)Bounds.GetCenter();
+		OutRenderData.Bounds.BoxExtent = (FVector)Bounds.GetExtent();
 
 		// Calculate the bounding sphere, using the center of the bounding box as the origin.
 		OutRenderData.Bounds.SphereRadius = 0.0f;
 		for (uint32 VertexIndex = 0; VertexIndex < BasePositionVertexBuffer.GetNumVertices(); VertexIndex++)
 		{
 			OutRenderData.Bounds.SphereRadius = FMath::Max(
-				(BasePositionVertexBuffer.VertexPosition(VertexIndex) - (FVector3f)OutRenderData.Bounds.Origin).Size(),
+				(BasePositionVertexBuffer.VertexPosition(VertexIndex) - (FVector3f)OutRenderData.Bounds.Origin).SizeSquared(),
 				OutRenderData.Bounds.SphereRadius
 				);
 		}
+		OutRenderData.Bounds.SphereRadius = FMath::Sqrt( OutRenderData.Bounds.SphereRadius );
 
 		Stage = EStage::GenerateRendering;
 		return true;
@@ -3079,7 +3082,7 @@ public:
 	virtual uint32 GetVertexIndex(uint32 FaceIndex, uint32 TriIndex) = 0;
 	virtual FVector3f GetVertexPosition(uint32 WedgeIndex) = 0;
 	virtual FVector3f GetVertexPosition(uint32 FaceIndex, uint32 TriIndex) = 0;
-	virtual FVector2D GetVertexUV(uint32 FaceIndex, uint32 TriIndex, uint32 UVIndex) = 0;
+	virtual FVector2f GetVertexUV(uint32 FaceIndex, uint32 TriIndex, uint32 UVIndex) = 0;
 	virtual uint32 GetFaceSmoothingGroups(uint32 FaceIndex) = 0;
 
 	virtual uint32 GetNumFaces() = 0;
@@ -3199,9 +3202,9 @@ public:
 		return Points[Wedges[Faces[FaceIndex].iWedge[TriIndex]].iVertex];
 	}
 
-	virtual FVector2D GetVertexUV(uint32 FaceIndex, uint32 TriIndex, uint32 UVIndex) override
+	virtual FVector2f GetVertexUV(uint32 FaceIndex, uint32 TriIndex, uint32 UVIndex) override
 	{
-		return FVector2D(Wedges[Faces[FaceIndex].iWedge[TriIndex]].UVs[UVIndex]);
+		return Wedges[Faces[FaceIndex].iWedge[TriIndex]].UVs[UVIndex];
 	}
 
 	virtual uint32 GetFaceSmoothingGroups(uint32 FaceIndex)
@@ -3294,7 +3297,7 @@ public:
 			for (int32 TriIndex = 0; TriIndex < 3; ++TriIndex)
 			{
 				uint32 Index = BuildData->GetWedgeIndex(FaceIndex, TriIndex);
-				new(VertIndexAndZ)FIndexAndZ(Index, (FVector)BuildData->GetVertexPosition(Index));
+				new(VertIndexAndZ)FIndexAndZ(Index, BuildData->GetVertexPosition(Index));
 			}
 		}
 
@@ -3460,7 +3463,7 @@ public:
 			for (int32 CornerB = 0; CornerB < 3; ++CornerB)
 			{
 				const FVector3f& CornerBPosition = BuildData->GetVertexPosition((FaceIdxB * 3) + CornerB);
-				if (PointsEqual((FVector)CornerAPosition, (FVector)CornerBPosition, BuildData->BuildOptions.OverlappingThresholds))
+				if (PointsEqual(CornerAPosition, CornerBPosition, BuildData->BuildOptions.OverlappingThresholds))
 				{
 					bFoundMatch = true;
 					break;
@@ -3560,9 +3563,9 @@ public:
 			}
 
 			// Don't process degenerate triangles.
-			if (PointsEqual((FVector)CornerPositions[0], (FVector)CornerPositions[1], BuildData->BuildOptions.OverlappingThresholds)
-				|| PointsEqual((FVector)CornerPositions[0], (FVector)CornerPositions[2], BuildData->BuildOptions.OverlappingThresholds)
-				|| PointsEqual((FVector)CornerPositions[1], (FVector)CornerPositions[2], BuildData->BuildOptions.OverlappingThresholds))
+			if (PointsEqual(CornerPositions[0], CornerPositions[1], BuildData->BuildOptions.OverlappingThresholds) ||
+				PointsEqual(CornerPositions[0], CornerPositions[2], BuildData->BuildOptions.OverlappingThresholds) ||
+				PointsEqual(CornerPositions[1], CornerPositions[2], BuildData->BuildOptions.OverlappingThresholds))
 			{
 				continue;
 			}
@@ -3648,8 +3651,8 @@ public:
 						for (int32 OtherCornerIndex = 0; OtherCornerIndex < 3; OtherCornerIndex++)
 						{
 							if (PointsEqual(
-								(FVector)CornerPositions[OurCornerIndex],
-								(FVector)BuildData->GetVertexPosition(OtherFaceIndex, OtherCornerIndex),
+								CornerPositions[OurCornerIndex],
+								BuildData->GetVertexPosition(OtherFaceIndex, OtherCornerIndex),
 								BuildData->BuildOptions.OverlappingThresholds
 								))
 							{
@@ -3706,8 +3709,8 @@ public:
 												int32 NextVertexIndex = BuildData->GetVertexIndex(NextFace.FaceIndex, NextCornerIndex);
 												int32 OtherVertexIndex = BuildData->GetVertexIndex(OtherFace.FaceIndex, OtherCornerIndex);
 												if (PointsEqual(
-													(FVector)BuildData->GetVertexPosition(NextFace.FaceIndex, NextCornerIndex),
-													(FVector)BuildData->GetVertexPosition(OtherFace.FaceIndex, OtherCornerIndex),
+													BuildData->GetVertexPosition(NextFace.FaceIndex, NextCornerIndex),
+													BuildData->GetVertexPosition(OtherFace.FaceIndex, OtherCornerIndex),
 													BuildData->BuildOptions.OverlappingThresholds))
 												{
 													CommonVertices++;
@@ -3782,9 +3785,9 @@ public:
 							if (bComputeWeightedNormals)
 							{
 								FVector3f OtherFacePoint[3] = { BuildData->GetVertexPosition(OtherFaceIndex, 0), BuildData->GetVertexPosition(OtherFaceIndex, 1), BuildData->GetVertexPosition(OtherFaceIndex, 2) };
-								float OtherFaceArea = TriangleUtilities::ComputeTriangleArea((FVector)OtherFacePoint[0], (FVector)OtherFacePoint[1], (FVector)OtherFacePoint[2]);
+								float OtherFaceArea = TriangleUtilities::ComputeTriangleArea(OtherFacePoint[0], OtherFacePoint[1], OtherFacePoint[2]);
 								int32 OtherFaceCornerIndex = RelevantFace.LinkedVertexIndex;
-								float OtherFaceAngle = TriangleUtilities::ComputeTriangleCornerAngle((FVector)OtherFacePoint[OtherFaceCornerIndex], (FVector)OtherFacePoint[(OtherFaceCornerIndex + 1) % 3], (FVector)OtherFacePoint[(OtherFaceCornerIndex + 2) % 3]);
+								float OtherFaceAngle = TriangleUtilities::ComputeTriangleCornerAngle(OtherFacePoint[OtherFaceCornerIndex], OtherFacePoint[(OtherFaceCornerIndex + 1) % 3], OtherFacePoint[(OtherFaceCornerIndex + 2) % 3]);
 								//Get the CornerWeight
 								CornerWeight = OtherFaceArea * OtherFaceAngle;
 							}
@@ -4218,7 +4221,7 @@ bool FMeshUtilities::BuildSkeletalMesh(FSkeletalMeshLODModel& LODModel,	const FS
 			for (int32 VertIndex = 0; VertIndex < NumSoftVertices; ++VertIndex)
 			{
 				FSoftSkinVertex& SrcVert = CurSection.SoftVertices[VertIndex];
-				new(VertIndexAndZ)FIndexAndZ(VertIndex, (FVector)SrcVert.Position);
+				new(VertIndexAndZ)FIndexAndZ(VertIndex, SrcVert.Position);
 			}
 			VertIndexAndZ.Sort(FCompareIndexAndZ());
 
@@ -4237,7 +4240,7 @@ bool FMeshUtilities::BuildSkeletalMesh(FSkeletalMeshLODModel& LODModel,	const FS
 
 					const uint32 IterVertIndex = VertIndexAndZ[j].Index;
 					FSoftSkinVertex& IterVert = CurSection.SoftVertices[IterVertIndex];
-					if (PointsEqual((FVector)SrcVert.Position, (FVector)IterVert.Position))
+					if (PointsEqual(SrcVert.Position, IterVert.Position))
 					{
 						// if so, we add to overlapping vert
 						TArray<int32>& SrcValueArray = CurSection.OverlappingVertices.FindOrAdd(SrcVertIndex);
@@ -4622,8 +4625,8 @@ bool FMeshUtilities::BuildSkeletalMesh_Legacy(FSkeletalMeshLODModel& LODModel
 
 				// check to see if the points are really overlapping
 				if (PointsEqual(
-					(FVector)Points[VertIndexAndZ[i].Index],
-					(FVector)Points[VertIndexAndZ[j].Index], OverlappingThresholds))
+					Points[VertIndexAndZ[i].Index],
+					Points[VertIndexAndZ[j].Index], OverlappingThresholds))
 				{
 					Vert2Duplicates.Add(VertIndexAndZ[i].Index, VertIndexAndZ[j].Index);
 					Vert2Duplicates.Add(VertIndexAndZ[j].Index, VertIndexAndZ[i].Index);
@@ -4716,7 +4719,7 @@ bool FMeshUtilities::BuildSkeletalMesh_Legacy(FSkeletalMeshLODModel& LODModel
 				const SkeletalMeshImportData::FMeshFace&	OtherFace = Faces[OtherFaceIndex];
 
 				FVector3f OtherFacePoint[3] = { Points[Wedges[OtherFace.iWedge[0]].iVertex], Points[Wedges[OtherFace.iWedge[1]].iVertex], Points[Wedges[OtherFace.iWedge[2]].iVertex] };
-				float OtherFaceArea = !bComputeWeightedNormals ? 1.0f : TriangleUtilities::ComputeTriangleArea((FVector)OtherFacePoint[0], (FVector)OtherFacePoint[1], (FVector)OtherFacePoint[2]);
+				float OtherFaceArea = !bComputeWeightedNormals ? 1.0f : TriangleUtilities::ComputeTriangleArea(OtherFacePoint[0], OtherFacePoint[1], OtherFacePoint[2]);
 				FVector3f		OtherTriangleNormal = (FVector3f)FPlane(
 					(FVector)OtherFacePoint[2],
 					(FVector)OtherFacePoint[1],
@@ -4729,13 +4732,13 @@ bool FMeshUtilities::BuildSkeletalMesh_Legacy(FSkeletalMeshLODModel& LODModel
 					for (int32 OtherVertexIndex = 0; OtherVertexIndex < 3; OtherVertexIndex++)
 					{
 						if (PointsEqual(
-							(FVector)OtherFacePoint[OtherVertexIndex],
-							(FVector)FacePoint[VertexIndex],
+							OtherFacePoint[OtherVertexIndex],
+							FacePoint[VertexIndex],
 							OverlappingThresholds
 							))
 						{
 							//Compute the angle
-							float OtherFaceAngle = !bComputeWeightedNormals ? 1.0f : TriangleUtilities::ComputeTriangleCornerAngle((FVector)OtherFacePoint[OtherVertexIndex], (FVector)OtherFacePoint[(OtherVertexIndex + 1) % 3], (FVector)OtherFacePoint[(OtherVertexIndex + 2) % 3]);
+							float OtherFaceAngle = !bComputeWeightedNormals ? 1.0f : TriangleUtilities::ComputeTriangleCornerAngle(OtherFacePoint[OtherVertexIndex], OtherFacePoint[(OtherVertexIndex + 1) % 3], OtherFacePoint[(OtherVertexIndex + 2) % 3]);
 							
 							float CornerWeight = (OtherFaceArea * OtherFaceAngle);
 
diff --git a/Engine/Source/Developer/MeshUtilities/Private/SkeletalMeshTools.cpp b/Engine/Source/Developer/MeshUtilities/Private/SkeletalMeshTools.cpp
--- a/Engine/Source/Developer/MeshUtilities/Private/SkeletalMeshTools.cpp
+++ b/Engine/Source/Developer/MeshUtilities/Private/SkeletalMeshTools.cpp
@@ -18,7 +18,7 @@ namespace SkeletalMeshTools
 {
 	bool AreSkelMeshVerticesEqual( const FSoftSkinBuildVertex& V1, const FSoftSkinBuildVertex& V2, const FOverlappingThresholds& OverlappingThresholds)
 	{
-		if(!PointsEqual((FVector)V1.Position, (FVector)V2.Position, OverlappingThresholds))
+		if(!PointsEqual(V1.Position, V2.Position, OverlappingThresholds))
 		{
 			return false;
 		}
@@ -31,17 +31,17 @@ namespace SkeletalMeshTools
 			}
 		}
 
-		if(!NormalsEqual((FVector)V1.TangentX, (FVector)V2.TangentX, OverlappingThresholds))
+		if(!NormalsEqual(V1.TangentX, V2.TangentX, OverlappingThresholds))
 		{
 			return false;
 		}
 
-		if(!NormalsEqual((FVector)V1.TangentY, (FVector)V2.TangentY, OverlappingThresholds))
+		if(!NormalsEqual(V1.TangentY, V2.TangentY, OverlappingThresholds))
 		{
 			return false;
 		}
 
-		if(!NormalsEqual((FVector)V1.TangentZ, (FVector)V2.TangentZ, OverlappingThresholds))
+		if(!NormalsEqual(V1.TangentZ, V2.TangentZ, OverlappingThresholds))
 		{
 			return false;
 		}
@@ -103,8 +103,8 @@ namespace SkeletalMeshTools
 
 					// check to see if the points are really overlapping
 					if(PointsEqual(
-						(FVector)RawVertices[RawVertIndexAndZ[i].Index].Position,
-						(FVector)RawVertices[RawVertIndexAndZ[j].Index].Position, OverlappingThresholds))
+						RawVertices[RawVertIndexAndZ[i].Index].Position,
+						RawVertices[RawVertIndexAndZ[j].Index].Position, OverlappingThresholds))
 					{
 						RawVerts2Dupes.Add(RawVertIndexAndZ[i].Index, RawVertIndexAndZ[j].Index);
 						RawVerts2Dupes.Add(RawVertIndexAndZ[j].Index, RawVertIndexAndZ[i].Index);
diff --git a/Engine/Source/Developer/SkeletalMeshUtilitiesCommon/Private/LODUtilities.cpp b/Engine/Source/Developer/SkeletalMeshUtilitiesCommon/Private/LODUtilities.cpp
--- a/Engine/Source/Developer/SkeletalMeshUtilitiesCommon/Private/LODUtilities.cpp
+++ b/Engine/Source/Developer/SkeletalMeshUtilitiesCommon/Private/LODUtilities.cpp
@@ -2202,7 +2202,7 @@ bool FLODUtilities::UpdateAlternateSkinWeights(FSkeletalMeshLODModel& LODModelDe
 	VertIndexAndZ.Reserve(VertexNumberDest);
 	for (int32 VertexIndex = 0; VertexIndex < VertexNumberDest; ++VertexIndex)
 	{
-		new(VertIndexAndZ)FIndexAndZ(VertexIndex, (FVector)ImportDataDest.Points[VertexIndex]);
+		new(VertIndexAndZ)FIndexAndZ(VertexIndex, ImportDataDest.Points[VertexIndex]);
 	}
 	// Sort the vertices by z value
 	VertIndexAndZ.Sort(FCompareIndexAndZ());
@@ -2210,7 +2210,7 @@ bool FLODUtilities::UpdateAlternateSkinWeights(FSkeletalMeshLODModel& LODModelDe
 	auto FindSimilarPosition = [&VertIndexAndZ, &ImportDataDest](const FVector3f& Position, TArray<int32>& PositionMatches, const float ComparisonThreshold)
 	{
 		PositionMatches.Reset();
-		FIndexAndZ PositionZ = FIndexAndZ(0, (FVector)Position);
+		FIndexAndZ PositionZ = FIndexAndZ(0, Position);
 		// Search for duplicates, quickly!
 		for (int32 i = 0; i < VertIndexAndZ.Num(); i++)
 		{
diff --git a/Engine/Source/Runtime/Engine/Public/MeshBuild.h b/Engine/Source/Runtime/Engine/Public/MeshBuild.h
--- a/Engine/Source/Runtime/Engine/Public/MeshBuild.h
+++ b/Engine/Source/Runtime/Engine/Public/MeshBuild.h
@@ -37,61 +37,41 @@ public:
 /**
  * Returns true if the specified points are about equal
  */
-inline bool PointsEqual(const FVector& V1,const FVector& V2, bool bUseEpsilonCompare = true )
+inline bool PointsEqual(const FVector3f& V1,const FVector3f& V2, bool bUseEpsilonCompare = true )
 {
 	const float Epsilon = bUseEpsilonCompare ? UE_THRESH_POINTS_ARE_SAME : 0.0f;
 	return FMath::Abs(V1.X - V2.X) <= Epsilon && FMath::Abs(V1.Y - V2.Y) <= Epsilon && FMath::Abs(V1.Z - V2.Z) <= Epsilon;
 }
 
-inline bool PointsEqual(const FVector3f& V1,const FVector3f& V2, bool bUseEpsilonCompare = true )
-{
-	return PointsEqual(FVector(V1), FVector(V2), bUseEpsilonCompare);
-}
-
-inline bool PointsEqual(const FVector& V1, const FVector& V2, const FOverlappingThresholds& OverlappingThreshold)
+inline bool PointsEqual(const FVector3f& V1, const FVector3f& V2, const FOverlappingThresholds& OverlappingThreshold)
 {
 	const float Epsilon = OverlappingThreshold.ThresholdPosition;
 	return FMath::Abs(V1.X - V2.X) <= Epsilon && FMath::Abs(V1.Y - V2.Y) <= Epsilon && FMath::Abs(V1.Z - V2.Z) <= Epsilon;
 }
 
-inline bool PointsEqual(const FVector3f& V1, const FVector3f& V2, const FOverlappingThresholds& OverlappingThreshold)
-{
-	return PointsEqual(FVector(V1), FVector(V2), OverlappingThreshold);
-}
-
 /**
  * Returns true if the specified normal vectors are about equal
  */
-inline bool NormalsEqual(const FVector& V1,const FVector& V2)
+inline bool NormalsEqual(const FVector3f& V1,const FVector3f& V2)
 {
 	const float Epsilon = UE_THRESH_NORMALS_ARE_SAME;
 	return FMath::Abs(V1.X - V2.X) <= Epsilon && FMath::Abs(V1.Y - V2.Y) <= Epsilon && FMath::Abs(V1.Z - V2.Z) <= Epsilon;
 }
 
-inline bool UVsEqual(const FVector2D& V1, const FVector2D& V2)
+inline bool UVsEqual(const FVector2f& V1, const FVector2f& V2)
 {
 	const float Epsilon = 1.0f / 1024.0f;
 	return FMath::Abs(V1.X - V2.X) <= Epsilon && FMath::Abs(V1.Y - V2.Y) <= Epsilon;
 }
 
-inline bool UVsEqual(const FVector2f& V1, const FVector2f& V2)
-{
-	return UVsEqual(FVector2D(V1), FVector2D(V2));
-}
-
-inline bool NormalsEqual(const FVector& V1,const FVector& V2, const FOverlappingThresholds& OverlappingThreshold)
+inline bool NormalsEqual(const FVector3f& V1,const FVector3f& V2, const FOverlappingThresholds& OverlappingThreshold)
 {
 	const float Epsilon = OverlappingThreshold.ThresholdTangentNormal;
 	return FMath::Abs(V1.X - V2.X) <= Epsilon && FMath::Abs(V1.Y - V2.Y) <= Epsilon && FMath::Abs(V1.Z - V2.Z) <= Epsilon;
 }
 
-inline bool UVsEqual(const FVector2D& V1, const FVector2D& V2, const FOverlappingThresholds& OverlappingThreshold)
+inline bool UVsEqual(const FVector2f& V1, const FVector2f& V2, const FOverlappingThresholds& OverlappingThreshold)
 {
 	const float Epsilon = OverlappingThreshold.ThresholdUV;
 	return FMath::Abs(V1.X - V2.X) <= Epsilon && FMath::Abs(V1.Y - V2.Y) <= Epsilon;
 }
-
-inline bool UVsEqual(const FVector2f& V1, const FVector2f& V2, const FOverlappingThresholds& OverlappingThreshold)
-{
-	return UVsEqual(FVector2D(V1), FVector2D(V2), OverlappingThreshold);
-}
diff --git a/Engine/Source/Runtime/Landscape/Private/LandscapeEdit.cpp b/Engine/Source/Runtime/Landscape/Private/LandscapeEdit.cpp
--- a/Engine/Source/Runtime/Landscape/Private/LandscapeEdit.cpp
+++ b/Engine/Source/Runtime/Landscape/Private/LandscapeEdit.cpp
@@ -3535,7 +3535,7 @@ bool ALandscapeProxy::ExportToRawMesh(
 					FVector Position = CDI.GetWorldVertex(VertexX, VertexY);
 
 					// If at least one vertex is within the given bounds we should process the quad  
-					new(VertIndexAndZ)FIndexAndZ(CurrentIndex, Position);
+					new(VertIndexAndZ)FIndexAndZ(CurrentIndex, (FVector3f)Position);
 					IndexToPosition.Add(CurrentIndex, Position);
 					CurrentIndex++;
 				}
@@ -3547,7 +3547,7 @@ bool ALandscapeProxy::ExportToRawMesh(
 		auto FindPreviousIndex = [&VertIndexAndZ, &IndexToPosition](int32 Index)->int32
 		{
 			const FVector& PositionA = IndexToPosition[Index];
-			FIndexAndZ CompressPosition(0, PositionA);
+			FIndexAndZ CompressPosition(0, (FVector3f)PositionA);
 			// Search for lowest index duplicates
 			int32 BestIndex = MAX_int32;
 			for (int32 i = 0; i < IndexToPosition.Num(); i++)
@@ -3558,7 +3558,7 @@ bool ALandscapeProxy::ExportToRawMesh(
 					break;
 				}
 				const FVector& PositionB = IndexToPosition[VertIndexAndZ[i].Index];
-				if (PointsEqual(PositionA, PositionB, SMALL_NUMBER))
+				if (PointsEqual((FVector3f)PositionA, (FVector3f)PositionB, SMALL_NUMBER))
 				{
 					if (VertIndexAndZ[i].Index < BestIndex)
 					{
diff --git a/Engine/Source/Runtime/MeshUtilitiesCommon/Private/LayoutUV.cpp b/Engine/Source/Runtime/MeshUtilitiesCommon/Private/LayoutUV.cpp
--- a/Engine/Source/Runtime/MeshUtilitiesCommon/Private/LayoutUV.cpp
+++ b/Engine/Source/Runtime/MeshUtilitiesCommon/Private/LayoutUV.cpp
@@ -249,13 +249,13 @@ int32 FLayoutUV::FChartFinder::FindCharts( const FOverlappingCorners& Overlappin
 		for( ; Tri < NumTris && DisjointSet[ SortedTris[ Tri ] ] == ChartID; Tri++ )
 		{
 			// Calculate chart bounds
-			FVector		Positions[3];
+			FVector3f	Positions[3];
 			FVector2f	UVs[3];
 			for( int k = 0; k < 3; k++ )
 			{
 				uint32 Index = 3 * SortedTris[ Tri ] + k;
 
-				Positions[k] = FVector(MeshView.GetPosition( Index ));
+				Positions[k] = MeshView.GetPosition( Index );
 				UVs[k] = TexCoords[ Index ];
 
 				Chart.MinUV.X = FMath::Min( Chart.MinUV.X, UVs[k].X );
@@ -264,8 +264,8 @@ int32 FLayoutUV::FChartFinder::FindCharts( const FOverlappingCorners& Overlappin
 				Chart.MaxUV.Y = FMath::Max( Chart.MaxUV.Y, UVs[k].Y );
 			}
 
-			FVector Edge1 = Positions[1] - Positions[0];
-			FVector Edge2 = Positions[2] - Positions[0];
+			FVector3f Edge1 = Positions[1] - Positions[0];
+			FVector3f Edge2 = Positions[2] - Positions[0];
 			float Area = 0.5f * ( Edge1 ^ Edge2 ).Size();
 
 			FVector2f EdgeUV1 = UVs[1] - UVs[0];
diff --git a/Engine/Source/Runtime/MeshUtilitiesCommon/Private/OverlappingCorners.cpp b/Engine/Source/Runtime/MeshUtilitiesCommon/Private/OverlappingCorners.cpp
--- a/Engine/Source/Runtime/MeshUtilitiesCommon/Private/OverlappingCorners.cpp
+++ b/Engine/Source/Runtime/MeshUtilitiesCommon/Private/OverlappingCorners.cpp
@@ -40,7 +40,7 @@ FOverlappingCorners::FOverlappingCorners(const FLayoutUV::IMeshView& MeshView, f
 	VertIndexAndZ.Reserve(NumWedges);
 	for (int32 WedgeIndex = 0; WedgeIndex < NumWedges; WedgeIndex++)
 	{
-		new(VertIndexAndZ)FIndexAndZ(WedgeIndex, FVector(MeshView.GetPosition(WedgeIndex)));
+		new(VertIndexAndZ)FIndexAndZ(WedgeIndex, MeshView.GetPosition(WedgeIndex));
 	}
 
 	// Sort the vertices by z value
diff --git a/Engine/Source/Runtime/MeshUtilitiesCommon/Public/MeshUtilitiesCommon.h b/Engine/Source/Runtime/MeshUtilitiesCommon/Public/MeshUtilitiesCommon.h
--- a/Engine/Source/Runtime/MeshUtilitiesCommon/Public/MeshUtilitiesCommon.h
+++ b/Engine/Source/Runtime/MeshUtilitiesCommon/Public/MeshUtilitiesCommon.h
@@ -33,18 +33,13 @@ struct FBoneVertInfo
 /** Helper struct for building acceleration structures. */
 struct FIndexAndZ
 {
-	FVector::FReal Z;
+	float Z;
 	int32 Index;
 
 	/** Default constructor. */
 	FIndexAndZ() {}
 
 	/** Initialization constructor. */
-	FIndexAndZ(int32 InIndex, FVector V)
-	{
-		Z = 0.30f * V.X + 0.33f * V.Y + 0.37f * V.Z;
-		Index = InIndex;
-	}
 	FIndexAndZ(int32 InIndex, FVector3f V)
 	{
 		Z = 0.30f * V.X + 0.33f * V.Y + 0.37f * V.Z;
@@ -61,17 +56,6 @@ struct FCompareIndexAndZ
 /**
 * Returns true if the specified points are about equal
 */
-inline bool PointsEqual(const FVector& V1, const FVector& V2, float ComparisonThreshold)
-{
-	if (FMath::Abs(V1.X - V2.X) > ComparisonThreshold
-		|| FMath::Abs(V1.Y - V2.Y) > ComparisonThreshold
-		|| FMath::Abs(V1.Z - V2.Z) > ComparisonThreshold)
-	{
-		return false;
-	}
-	return true;
-}
-
 inline bool PointsEqual(const FVector3f& V1, const FVector3f& V2, float ComparisonThreshold)
 {
 	if (FMath::Abs(V1.X - V2.X) > ComparisonThreshold
@@ -88,24 +72,25 @@ namespace TriangleUtilities
 	/*
 	 * This function compute the area of a triangle, it will return zero if the triangle is degenerated
 	 */
-	static double ComputeTriangleArea(const FVector& PointA, const FVector& PointB, const FVector& PointC)
+	static float ComputeTriangleArea(const FVector3f& PointA, const FVector3f& PointB, const FVector3f& PointC)
 	{
-		return FVector::CrossProduct((PointB - PointA), (PointC - PointA)).Size() / 2.0f;
+		return FVector3f::CrossProduct((PointB - PointA), (PointC - PointA)).Size() / 2.0f;
 	}
 
 	/*
 	 * This function compute the angle of a triangle corner, it will return zero if the triangle is degenerated
 	 */
-	static double ComputeTriangleCornerAngle(const FVector& PointA, const FVector& PointB, const FVector& PointC)
+	static float ComputeTriangleCornerAngle(const FVector3f& PointA, const FVector3f& PointB, const FVector3f& PointC)
 	{
-		FVector E1 = (PointB - PointA);
-		FVector E2 = (PointC - PointA);
+		FVector3f E1 = (PointB - PointA);
+		FVector3f E2 = (PointC - PointA);
 		//Normalize both edges (unit vector) of the triangle so we get a dotProduct result that will be a valid acos input [-1, 1]
 		if (!E1.Normalize() || !E2.Normalize())
 		{
 			//Return a null ratio if the polygon is degenerate
 			return 0.0f;
 		}
-		return FMath::Acos(FVector::DotProduct(E1, E2));
+		float DotProduct = FVector3f::DotProduct(E1, E2);
+		return FMath::Acos(DotProduct);
 	}
 }
\ No newline at end of file
diff --git a/Engine/Source/Runtime/StaticMeshDescription/Private/StaticMeshOperations.cpp b/Engine/Source/Runtime/StaticMeshDescription/Private/StaticMeshOperations.cpp
--- a/Engine/Source/Runtime/StaticMeshDescription/Private/StaticMeshOperations.cpp
+++ b/Engine/Source/Runtime/StaticMeshDescription/Private/StaticMeshOperations.cpp
@@ -1356,9 +1356,9 @@ void FStaticMeshOperations::ComputeTangentsAndNormals(FMeshDescription& MeshDesc
 
 		for (int32 Index = 0, TriIndex = 0; Index < NumTriangles; Index++, TriIndex += 3)
 		{
-			const FVector PointA(VertexPositions[TriVertexIDs[TriIndex + 0]]);
-			const FVector PointB(VertexPositions[TriVertexIDs[TriIndex + 1]]);
-			const FVector PointC(VertexPositions[TriVertexIDs[TriIndex + 2]]);
+			const FVector3f PointA(VertexPositions[TriVertexIDs[TriIndex + 0]]);
+			const FVector3f PointB(VertexPositions[TriVertexIDs[TriIndex + 1]]);
+			const FVector3f PointC(VertexPositions[TriVertexIDs[TriIndex + 2]]);
 			FTriangleData& TriangleData = TriangleDatas[Index];
 			TriangleData.Area = TriangleUtilities::ComputeTriangleArea(PointA, PointB, PointC);
 			TriangleData.SetCornerAngleData(TriVertexInstanceIDs[TriIndex + 0], TriangleUtilities::ComputeTriangleCornerAngle(PointA, PointB, PointC), 0);
