MeshUtilities: Fix an issue where the "Make Static Mesh" button was missing

diff --git a/Engine/Source/Developer/MeshUtilities/Private/MeshUtilities.cpp b/Engine/Source/Developer/MeshUtilities/Private/MeshUtilities.cpp
--- a/Engine/Source/Developer/MeshUtilities/Private/MeshUtilities.cpp
+++ b/Engine/Source/Developer/MeshUtilities/Private/MeshUtilities.cpp
@@ -6123,23 +6123,23 @@ void FMeshUtilities::AddMakeStaticMeshEntryToToolMenu(FName InToolMenuName)
 	UToolMenu* ToolMenu = UToolMenus::Get()->ExtendMenu(InToolMenuName);
 	FToolMenuSection& SkeletalMeshSection = ToolMenu->AddSection("SkeletalMesh", LOCTEXT("SkeletalMeshSection", "Skeletal Mesh"));
 
-	SkeletalMeshSection.AddMenuEntry(
-		"MakeStaticMesh",
-		LOCTEXT("MakeStaticMesh", "Make Static Mesh"),
-		LOCTEXT("MakeStaticMeshTooltip", "Make a new static mesh out of the preview's current pose."),
-		FSlateIcon(FAppStyle::GetAppStyleSetName(), "Persona.ConvertToStaticMesh"),
-		FToolUIActionChoice(
-			FToolUIAction(FToolMenuExecuteAction::CreateLambda([this](const FToolMenuContext& InContext)
-			{
-				if(UPersonaToolMenuContext* MenuContext = InContext.FindContext<UPersonaToolMenuContext>())
-				{
-					if(UDebugSkelMeshComponent* MeshComponent = MenuContext->GetPreviewMeshComponent())
-					{
-						ConvertMeshesToStaticMesh(TArray<UMeshComponent*>({ MeshComponent }), MeshComponent->GetComponentToWorld());
-					}
-				}
-			}))
-		)
+	SkeletalMeshSection.AddEntry(FToolMenuEntry::InitToolBarButton(
+        "MakeStaticMesh",
+        FToolUIActionChoice(
+            FToolUIAction(FToolMenuExecuteAction::CreateLambda([this](const FToolMenuContext& InContext)
+            {
+            	if(UPersonaToolMenuContext* MenuContext = InContext.FindContext<UPersonaToolMenuContext>())
+            	{
+            		if(UDebugSkelMeshComponent* MeshComponent = MenuContext->GetPreviewMeshComponent())
+            		{
+            			ConvertMeshesToStaticMesh(TArray<UMeshComponent*>({ MeshComponent }), MeshComponent->GetComponentToWorld());
+            		}
+            	}
+            }))
+        ),
+        LOCTEXT("MakeStaticMesh", "Make Static Mesh"),
+        LOCTEXT("MakeStaticMeshTooltip", "Make a new static mesh out of the preview's current pose."),
+        FSlateIcon(FAppStyle::GetAppStyleSetName(), "Persona.ConvertToStaticMesh"))
 	);
 }
 
