Fix issue causing incorrect Mesh SDF generation
- Workaround issue with compiler optimization by using copy constructor.

diff --git a/Engine/Source/Developer/MeshUtilities/Private/MeshRepresentationCommon.cpp b/Engine/Source/Developer/MeshUtilities/Private/MeshRepresentationCommon.cpp
--- a/Engine/Source/Developer/MeshUtilities/Private/MeshRepresentationCommon.cpp
+++ b/Engine/Source/Developer/MeshUtilities/Private/MeshRepresentationCommon.cpp
@@ -9,7 +9,7 @@ static FVector3f UniformSampleHemisphere(FVector2D Uniforms)
 {
 	Uniforms = Uniforms * 2.0f - 1.0f;
 
-	if(Uniforms == FVector2D::ZeroVector)
+	if (Uniforms == FVector2D::ZeroVector)
 	{
 		return FVector3f::ZeroVector;
 	}
@@ -53,7 +53,10 @@ void MeshUtilities::GenerateStratifiedUniformHemisphereSamples(int32 NumSamples,
 			const float Fraction1 = (IndexX + U1) / (float)NumSamplesDim;
 			const float Fraction2 = (IndexY + U2) / (float)NumSamplesDim;
 
-			Samples.Add(UniformSampleHemisphere(FVector2D(Fraction1, Fraction2)));
+			FVector3f Tmp = UniformSampleHemisphere(FVector2D(Fraction1, Fraction2));
+
+			// Workaround issue with compiler optimization by using copy constructor here.
+			Samples.Add(FVector3f(Tmp));
 		}
 	}
 }
diff --git a/Engine/Source/Runtime/Engine/Private/DistanceFieldAtlas.cpp b/Engine/Source/Runtime/Engine/Private/DistanceFieldAtlas.cpp
--- a/Engine/Source/Runtime/Engine/Private/DistanceFieldAtlas.cpp
+++ b/Engine/Source/Runtime/Engine/Private/DistanceFieldAtlas.cpp
@@ -132,7 +132,7 @@ FDistanceFieldAsyncQueue* GDistanceFieldAsyncQueue = NULL;
 #if WITH_EDITOR
 
 // DDC key for distance field data, must be changed when modifying the generation code or data format
-#define DISTANCEFIELD_DERIVEDDATA_VER TEXT("8A773878-E42B-4143-A23B-E231679CE486")
+#define DISTANCEFIELD_DERIVEDDATA_VER TEXT("295895D9-FFCE-48B5-9A4B-131DB318441E")
 
 FString BuildDistanceFieldDerivedDataKey(const FString& InMeshKey)
 {
diff --git a/Engine/Source/Runtime/Engine/Private/MeshCardRepresentation.cpp b/Engine/Source/Runtime/Engine/Private/MeshCardRepresentation.cpp
--- a/Engine/Source/Runtime/Engine/Private/MeshCardRepresentation.cpp
+++ b/Engine/Source/Runtime/Engine/Private/MeshCardRepresentation.cpp
@@ -105,7 +105,7 @@ FCardRepresentationAsyncQueue* GCardRepresentationAsyncQueue = NULL;
 #if WITH_EDITOR
 
 // DDC key for card representation data, must be changed when modifying the generation code or data format
-#define CARDREPRESENTATION_DERIVEDDATA_VER TEXT("3DAE2659-E697-429C-A046-4E6B398ECCBF")
+#define CARDREPRESENTATION_DERIVEDDATA_VER TEXT("FF1E9B99-1837-4F13-A892-13BD62922D0B")
 
 FString BuildCardRepresentationDerivedDataKey(const FString& InMeshKey, int32 MaxLumenMeshCards)
 {
