Load the MeshBoneReduction module as soon as the MeshUtilities module starts up, as it could be needed at any time from a background thread

diff --git a/Engine/Source/Developer/MeshUtilities/Private/MeshUtilities.cpp b/Engine/Source/Developer/MeshUtilities/Private/MeshUtilities.cpp
--- a/Engine/Source/Developer/MeshUtilities/Private/MeshUtilities.cpp
+++ b/Engine/Source/Developer/MeshUtilities/Private/MeshUtilities.cpp
@@ -5577,6 +5577,7 @@ void FMeshUtilities::StartupModule()
 {
 	FModuleManager::Get().LoadModule("MaterialBaking");
 	FModuleManager::Get().LoadModule(TEXT("MeshMergeUtilities"));
+	FModuleManager::Get().LoadModule(TEXT("MeshBoneReduction"));
 
 	FPropertyEditorModule& PropertyEditorModule = FModuleManager::Get().LoadModuleChecked<FPropertyEditorModule>("PropertyEditor");
 
