Lumen Mesh Cards - fixed hang when generating poorly constructred meshes and due to an EmbreeRay bug, EmbreeRay.ray.tfar.ray.tnear was further than EmbreeRay.ray.tfar

diff --git a/Engine/Source/Developer/MeshUtilities/Private/MeshCardRepresentationUtilities.cpp b/Engine/Source/Developer/MeshUtilities/Private/MeshCardRepresentationUtilities.cpp
--- a/Engine/Source/Developer/MeshUtilities/Private/MeshCardRepresentationUtilities.cpp
+++ b/Engine/Source/Developer/MeshUtilities/Private/MeshCardRepresentationUtilities.cpp
@@ -428,7 +428,9 @@ void GenerateSurfelsForDirection(
 
 						// Move ray to the next intersection
 						LastHitCoordZ = HitCoordZ;
-						RayTNear = std::nextafter(FMath::Max(NearPlaneOffset + (LastHitCoordZ + 1) * ClusteringParams.VoxelSize, EmbreeRay.ray.tfar), std::numeric_limits<float>::infinity());
+						// Sometimes EmbreeRay.ray.tnear was further than EmbreeRay.ray.tfar causing an infinite loop
+						const float SafeTFar = FMath::Max(EmbreeRay.ray.tfar, EmbreeRay.ray.tnear);
+						RayTNear = std::nextafter(FMath::Max(NearPlaneOffset + (LastHitCoordZ + 1) * ClusteringParams.VoxelSize, SafeTFar), std::numeric_limits<float>::infinity());
 						SkipPrimId = EmbreeRay.hit.primID;
 					}
 					else
