Fix StaticMeshToRawMeshes() not handling InstancedStaticMeshComponents properly. This fixes the "Convert to Static Mesh" tools.

diff --git a/Engine/Source/Developer/MeshUtilities/Private/MeshUtilities.cpp b/Engine/Source/Developer/MeshUtilities/Private/MeshUtilities.cpp
--- a/Engine/Source/Developer/MeshUtilities/Private/MeshUtilities.cpp
+++ b/Engine/Source/Developer/MeshUtilities/Private/MeshUtilities.cpp
@@ -21,6 +21,7 @@
 #include "SkeletalMeshToolMenuContext.h"
 #include "Components/MeshComponent.h"
 #include "RawIndexBuffer.h"
+#include "Components/InstancedStaticMeshComponent.h"
 #include "Components/StaticMeshComponent.h"
 #include "Components/ShapeComponent.h"
 #include "Engine/SkinnedAssetCommon.h"
@@ -423,6 +424,20 @@ static void StaticMeshToRawMeshes(UStaticMeshComponent* InStaticMeshComponent, i
 
 	const int32 NumLODs = InStaticMeshComponent->GetStaticMesh()->GetRenderData()->LODResources.Num();
 
+	const UInstancedStaticMeshComponent* ISMC = Cast<UInstancedStaticMeshComponent>(InStaticMeshComponent);
+	const int32 NumInstances = ISMC ? ISMC->GetInstanceCount() : 1;
+	auto GetInstanceTransformMatrix = [ISMC, InStaticMeshComponent, NumInstances](int32 InstanceIndex)
+		{
+			if (ISMC && InstanceIndex <= NumInstances)
+			{
+				FTransform InstanceTransform;
+				ISMC->GetInstanceTransform(InstanceIndex, InstanceTransform);
+				return InstanceTransform.ToMatrixWithScale();
+			}
+
+			return FMatrix::Identity;
+		};
+
 	for (int32 OverallLODIndex = 0; OverallLODIndex < InOverallMaxLODs; OverallLODIndex++)
 	{
 		int32 LODIndexRead = FMath::Min(OverallLODIndex, NumLODs - 1);
@@ -430,62 +445,67 @@ static void StaticMeshToRawMeshes(UStaticMeshComponent* InStaticMeshComponent, i
 		FRawMesh& RawMesh = OutRawMeshes[OverallLODIndex];
 		FRawMeshTracker& RawMeshTracker = OutRawMeshTrackers[OverallLODIndex];
 		const FStaticMeshLODResources& LODResource = InStaticMeshComponent->GetStaticMesh()->GetRenderData()->LODResources[LODIndexRead];
-		const int32 BaseVertexIndex = RawMesh.VertexPositions.Num();
 
-		for (int32 VertIndex = 0; VertIndex < LODResource.GetNumVertices(); ++VertIndex)
+		for (int32 InstanceIndex = 0; InstanceIndex < NumInstances; InstanceIndex++)
 		{
-			RawMesh.VertexPositions.Add(FVector4f(InComponentToWorld.TransformPosition((FVector)LODResource.VertexBuffers.PositionVertexBuffer.VertexPosition((uint32)VertIndex))));
-		}
+			const FMatrix InstanceToWorld = GetInstanceTransformMatrix(InstanceIndex) * InComponentToWorld;
+			const int32 BaseVertexIndex = RawMesh.VertexPositions.Num();
 
-		const FIndexArrayView IndexArrayView = LODResource.IndexBuffer.GetArrayView();
-		const FStaticMeshVertexBuffer& StaticMeshVertexBuffer = LODResource.VertexBuffers.StaticMeshVertexBuffer;
-		const int32 NumTexCoords = FMath::Min(StaticMeshVertexBuffer.GetNumTexCoords(), (uint32)MAX_MESH_TEXTURE_COORDS);
-		const int32 NumSections = LODResource.Sections.Num();
+			for (int32 VertIndex = 0; VertIndex < LODResource.GetNumVertices(); ++VertIndex)
+			{
+				RawMesh.VertexPositions.Add(FVector4f(InstanceToWorld.TransformPosition((FVector)LODResource.VertexBuffers.PositionVertexBuffer.VertexPosition((uint32)VertIndex))));
+			}
 
-		for (int32 SectionIndex = 0; SectionIndex < NumSections; SectionIndex++)
-		{
-			const FStaticMeshSection& StaticMeshSection = LODResource.Sections[SectionIndex];
+			const FIndexArrayView IndexArrayView = LODResource.IndexBuffer.GetArrayView();
+			const FStaticMeshVertexBuffer& StaticMeshVertexBuffer = LODResource.VertexBuffers.StaticMeshVertexBuffer;
+			const int32 NumTexCoords = FMath::Min(StaticMeshVertexBuffer.GetNumTexCoords(), (uint32)MAX_MESH_TEXTURE_COORDS);
+			const int32 NumSections = LODResource.Sections.Num();
 
-			const int32 NumIndices = StaticMeshSection.NumTriangles * 3;
-			for (int32 IndexIndex = 0; IndexIndex < NumIndices; IndexIndex++)
+			for (int32 SectionIndex = 0; SectionIndex < NumSections; SectionIndex++)
 			{
-				int32 Index = IndexArrayView[StaticMeshSection.FirstIndex + IndexIndex];
-				RawMesh.WedgeIndices.Add(BaseVertexIndex + Index);
+				const FStaticMeshSection& StaticMeshSection = LODResource.Sections[SectionIndex];
 
-				RawMesh.WedgeTangentX.Add(FVector4f(InComponentToWorld.TransformVector(FVector(StaticMeshVertexBuffer.VertexTangentX(Index)))));
-				RawMesh.WedgeTangentY.Add(FVector4f(InComponentToWorld.TransformVector(FVector(StaticMeshVertexBuffer.VertexTangentY(Index)))));
-				RawMesh.WedgeTangentZ.Add(FVector4f(InComponentToWorld.TransformVector(FVector(StaticMeshVertexBuffer.VertexTangentZ(Index)))));
-
-				for (int32 TexCoordIndex = 0; TexCoordIndex < MAX_MESH_TEXTURE_COORDS; TexCoordIndex++)
+				const int32 NumIndices = StaticMeshSection.NumTriangles * 3;
+				for (int32 IndexIndex = 0; IndexIndex < NumIndices; IndexIndex++)
 				{
-					if (TexCoordIndex >= NumTexCoords)
+					int32 Index = IndexArrayView[StaticMeshSection.FirstIndex + IndexIndex];
+					RawMesh.WedgeIndices.Add(BaseVertexIndex + Index);
+
+					RawMesh.WedgeTangentX.Add(FVector4f(InstanceToWorld.TransformVector(FVector(StaticMeshVertexBuffer.VertexTangentX(Index)))));
+					RawMesh.WedgeTangentY.Add(FVector4f(InstanceToWorld.TransformVector(FVector(StaticMeshVertexBuffer.VertexTangentY(Index)))));
+					RawMesh.WedgeTangentZ.Add(FVector4f(InstanceToWorld.TransformVector(FVector(StaticMeshVertexBuffer.VertexTangentZ(Index)))));
+
+					for (int32 TexCoordIndex = 0; TexCoordIndex < MAX_MESH_TEXTURE_COORDS; TexCoordIndex++)
+					{
+						if (TexCoordIndex >= NumTexCoords)
+						{
+							RawMesh.WedgeTexCoords[TexCoordIndex].AddDefaulted();
+						}
+						else
+						{
+							RawMesh.WedgeTexCoords[TexCoordIndex].Add(StaticMeshVertexBuffer.GetVertexUV(Index, TexCoordIndex));
+							RawMeshTracker.bValidTexCoords[TexCoordIndex] = true;
+						}
+					}
+
+					if (LODResource.VertexBuffers.ColorVertexBuffer.IsInitialized())
 					{
-						RawMesh.WedgeTexCoords[TexCoordIndex].AddDefaulted();
+						RawMesh.WedgeColors.Add(LODResource.VertexBuffers.ColorVertexBuffer.VertexColor(Index));
+						RawMeshTracker.bValidColors = true;
 					}
 					else
 					{
-						RawMesh.WedgeTexCoords[TexCoordIndex].Add(StaticMeshVertexBuffer.GetVertexUV(Index, TexCoordIndex));
-						RawMeshTracker.bValidTexCoords[TexCoordIndex] = true;
+						RawMesh.WedgeColors.Add(FColor::White);
 					}
 				}
 
-				if (LODResource.VertexBuffers.ColorVertexBuffer.IsInitialized())
-				{
-					RawMesh.WedgeColors.Add(LODResource.VertexBuffers.ColorVertexBuffer.VertexColor(Index));
-					RawMeshTracker.bValidColors = true;
-				}
-				else
+				// copy face info
+				for (uint32 TriIndex = 0; TriIndex < StaticMeshSection.NumTriangles; TriIndex++)
 				{
-					RawMesh.WedgeColors.Add(FColor::White);
+					RawMesh.FaceMaterialIndices.Add(BaseMaterialIndex + StaticMeshSection.MaterialIndex);
+					RawMesh.FaceSmoothingMasks.Add(0); // Assume this is ignored as bRecomputeNormals is false
 				}
 			}
-
-			// copy face info
-			for (uint32 TriIndex = 0; TriIndex < StaticMeshSection.NumTriangles; TriIndex++)
-			{
-				RawMesh.FaceMaterialIndices.Add(BaseMaterialIndex + StaticMeshSection.MaterialIndex);
-				RawMesh.FaceSmoothingMasks.Add(0); // Assume this is ignored as bRecomputeNormals is false
-			}
 		}
 	}
 
