Embree device/scene is used even after they fail init. Fix crash by early out. (This matches the behavior of GenerateCardRepresentationData).
Failed Embree init will leave FDistanceFieldVolumeData::IsValid() returning false.

#rnx

diff --git a/Engine/Source/Developer/MeshUtilities/Private/MeshDistanceFieldUtilities.cpp b/Engine/Source/Developer/MeshUtilities/Private/MeshDistanceFieldUtilities.cpp
--- a/Engine/Source/Developer/MeshUtilities/Private/MeshDistanceFieldUtilities.cpp
+++ b/Engine/Source/Developer/MeshUtilities/Private/MeshDistanceFieldUtilities.cpp
@@ -274,6 +274,12 @@ void FMeshUtilities::GenerateSignedDistanceFieldVolumeData(
 
 		check(EmbreeScene.bUseEmbree);
 
+		// If Embree setup fails, there will be no scene to operate on. Early out.
+		if (!EmbreeScene.EmbreeScene)
+		{
+			return;
+		}
+
 		// Whether to use an Embree Point Query to compute the closest unsigned distance.  Rays will only be traced to determine backfaces visible for sign.
 		const bool bUsePointQuery = true;
 
