Fix for alternate skin weights blowing up.

#rnx

diff --git a/Engine/Source/Developer/MeshUtilities/Private/MeshUtilities.cpp b/Engine/Source/Developer/MeshUtilities/Private/MeshUtilities.cpp
--- a/Engine/Source/Developer/MeshUtilities/Private/MeshUtilities.cpp
+++ b/Engine/Source/Developer/MeshUtilities/Private/MeshUtilities.cpp
@@ -5835,6 +5835,12 @@ void FMeshUtilities::GenerateRuntimeSkinWeightData(
 				{
 					OverrideIndex = UniqueWeights.Num();
 
+					// When writing out 8-bit weights, we want to collect the resulting deviation from fully normalized
+					// weight sum. Once all values are copied, we alter the first, and most prominent, weight, by the 
+					// deviation value so that we end up with completely normalized 8-bit weights.
+					const int32 FirstWeightOffset = InOutSkinWeightOverrideData.BoneWeights.Num();
+					int32 Accumulated8BitDeviation = 255;
+
 					// Write out non-zero weighted influences only
 					for (int32 InfluenceIndex = 0; InfluenceIndex < NumInfluences; ++InfluenceIndex)
 					{
@@ -5861,10 +5867,16 @@ void FMeshUtilities::GenerateRuntimeSkinWeightData(
 						else
 						{
 							BoneWeights.Add(static_cast<uint8>(Weight >> 8));
+							Accumulated8BitDeviation -= BoneWeights.Last();
 						}
 					}
 
 					UniqueWeights.Add(SourceSkinWeight);
+
+					if (!bInUseHighPrecisionWeights)
+					{
+						InOutSkinWeightOverrideData.BoneWeights[FirstWeightOffset] += Accumulated8BitDeviation;
+					}
 				}
 
 				InOutSkinWeightOverrideData.VertexIndexToInfluenceOffset.Add(VertexIndex, OverrideIndex);
diff --git a/Engine/Source/Runtime/Core/Private/UObject/FortniteMainBranchObjectVersion.cpp b/Engine/Source/Runtime/Core/Private/UObject/FortniteMainBranchObjectVersion.cpp
--- a/Engine/Source/Runtime/Core/Private/UObject/FortniteMainBranchObjectVersion.cpp
+++ b/Engine/Source/Runtime/Core/Private/UObject/FortniteMainBranchObjectVersion.cpp
@@ -12,7 +12,7 @@ TMap<FGuid, FGuid> FFortniteMainBranchObjectVersion::GetSystemGuids()
 	SystemGuids.Add(DevGuids.NANITE_DERIVEDDATA_VER, FGuid("D28C003F825F4C178F698938FB0E47F7"));
 	SystemGuids.Add(DevGuids.NIAGARASHADERMAP_DERIVEDDATA_VER, FGuid("35D7960B72164B49B9B7B3FB5857E674"));
 	SystemGuids.Add(DevGuids.Niagara_LatestScriptCompileVersion, FGuid("3A499F83F8274468B50A487A3349C2BB"));
-	SystemGuids.Add(DevGuids.SkeletalMeshDerivedDataVersion, FGuid("3E35053A0A87964939C750BB8234AAE2"));
+	SystemGuids.Add(DevGuids.SkeletalMeshDerivedDataVersion, FGuid("87620CAAE42ACAA2D32409D2AD939201"));
 	SystemGuids.Add(DevGuids.STATICMESH_DERIVEDDATA_VER, FGuid("543843B103794E0E9BA4BE80FB602F79"));
 	return SystemGuids;
 }
