#UE disable a warning, for now

diff --git a/Engine/Source/Developer/MeshUtilities/Private/MeshRepresentationCommon.cpp b/Engine/Source/Developer/MeshUtilities/Private/MeshRepresentationCommon.cpp
--- a/Engine/Source/Developer/MeshUtilities/Private/MeshRepresentationCommon.cpp
+++ b/Engine/Source/Developer/MeshUtilities/Private/MeshRepresentationCommon.cpp
@@ -165,10 +165,12 @@ void MeshRepresentation::SetupEmbreeScene(
 	}
 #endif
 
+	/*
 	if (LODModel.Sections.Num() > SectionData.Num())
 	{
 		UE_LOG(LogMeshUtilities, Warning, TEXT("Unexpected number of mesh sections when setting up Embree Scene for %s."), *MeshName);
 	}
+	*/
 
 	TArray<int32> FilteredTriangles;
 	FilteredTriangles.Empty(NumTriangles);
