Use FSignedDistanceFieldBuildSectionData when generating distance fields for Nanite meshes.

Context:
- Nanite meshes distance fields are generated using the original mesh data instead of the coarse representation generated by Nanite, however that code path was not checking section data such as blend mode, two sided, etc.

Fix:
- Modified FStaticMeshBuilder::BuildMeshVertexPositions(...) to also output a FStaticMeshSectionArray.
- Updated MeshRepresentation::SetupEmbreeScene(...) to check section flags when using SourceMeshData.IsValid() code path.

#prelight

diff --git a/Engine/Source/Developer/MeshBuilder/Private/MeshBuilderModule.cpp b/Engine/Source/Developer/MeshBuilder/Private/MeshBuilderModule.cpp
--- a/Engine/Source/Developer/MeshBuilder/Private/MeshBuilderModule.cpp
+++ b/Engine/Source/Developer/MeshBuilder/Private/MeshBuilderModule.cpp
@@ -31,7 +31,8 @@ public:
 	virtual bool BuildMeshVertexPositions(
 		UObject* StaticMesh,
 		TArray<uint32>& Indices,
-		TArray<FVector3f>& Vertices) override;
+		TArray<FVector3f>& Vertices,
+		FStaticMeshSectionArray& Sections) override;
 
 	virtual bool BuildSkeletalMesh(const FSkeletalMeshBuildParameters& SkeletalMeshBuildParameters) override;
 
@@ -55,13 +56,14 @@ bool FMeshBuilderModule::BuildMesh(FStaticMeshRenderData& OutRenderData, class U
 bool FMeshBuilderModule::BuildMeshVertexPositions(
 	UObject* Mesh,
 	TArray<uint32>& Indices,
-	TArray<FVector3f>& Vertices)
+	TArray<FVector3f>& Vertices,
+	FStaticMeshSectionArray& Sections)
 {
 	UStaticMesh* StaticMesh = Cast< UStaticMesh >(Mesh);
 	if (StaticMesh)
 	{
 		//Call the static mesh builder
-		return FStaticMeshBuilder().BuildMeshVertexPositions(StaticMesh, Indices, Vertices);
+		return FStaticMeshBuilder().BuildMeshVertexPositions(StaticMesh, Indices, Vertices, Sections);
 	}
 	return false;
 }
diff --git a/Engine/Source/Developer/MeshBuilder/Private/StaticMeshBuilder.cpp b/Engine/Source/Developer/MeshBuilder/Private/StaticMeshBuilder.cpp
--- a/Engine/Source/Developer/MeshBuilder/Private/StaticMeshBuilder.cpp
+++ b/Engine/Source/Developer/MeshBuilder/Private/StaticMeshBuilder.cpp
@@ -730,7 +730,8 @@ bool FStaticMeshBuilder::Build(FStaticMeshRenderData& StaticMeshRenderData, USta
 bool FStaticMeshBuilder::BuildMeshVertexPositions(
 	UStaticMesh* StaticMesh,
 	TArray<uint32>& BuiltIndices,
-	TArray<FVector3f>& BuiltVertices)
+	TArray<FVector3f>& BuiltVertices,
+	FStaticMeshSectionArray& Sections)
 {
 	TRACE_CPUPROFILER_EVENT_SCOPE(FStaticMeshBuilder::BuildMeshVertexPositions);
 
@@ -742,59 +743,62 @@ bool FStaticMeshBuilder::BuildMeshVertexPositions(
 		return false;
 	}
 
-	const int32 NumSourceModels = StaticMesh->GetNumSourceModels();
-	if (NumSourceModels > 0)
-	{
-		FMeshDescription MeshDescription;
-		const bool bIsMeshDescriptionValid = SourceModel.CloneMeshDescription(MeshDescription);
-		if (bIsMeshDescriptionValid)
-		{
-			FElementIDRemappings Remappings;
-			if (MeshDescription.NeedsCompact())
-			{
-				MeshDescription.Compact(Remappings);
-			}
+	FMeshDescription MeshDescription;
+	const bool bIsMeshDescriptionValid = SourceModel.CloneMeshDescription(MeshDescription);
+	check(bIsMeshDescriptionValid);
 
-			const FMeshBuildSettings& BuildSettings = SourceModel.BuildSettings;
+	FMeshBuildSettings& BuildSettings = StaticMesh->GetSourceModel(0).BuildSettings;
 
-			const FStaticMeshConstAttributes Attributes(MeshDescription);
-			TArrayView<const FVector3f> VertexPositions = Attributes.GetVertexPositions().GetRawArray();
-			TArrayView<const FVertexID> VertexIndices = Attributes.GetTriangleVertexIndices().GetRawArray();
-			const FVector3f BuildScale3D = (FVector3f)BuildSettings.BuildScale3D;
+	FMeshDescriptionHelper MeshDescriptionHelper(&BuildSettings);
+	MeshDescriptionHelper.SetupRenderMeshDescription(StaticMesh, MeshDescription, false, false);
 
-			BuiltVertices.Reserve(VertexPositions.Num());
-			for (int32 VertexIndex = 0; VertexIndex < VertexPositions.Num(); ++VertexIndex)
-			{
-				BuiltVertices.Add(VertexPositions[VertexIndex] * BuildScale3D);
-			}
+	const FPolygonGroupArray& PolygonGroups = MeshDescription.PolygonGroups();
 
-			BuiltIndices.Reserve(VertexIndices.Num());
-			for (int32 TriangleIndex = 0; TriangleIndex < VertexIndices.Num() / 3; ++TriangleIndex)
-			{
-				const uint32 I0 = VertexIndices[TriangleIndex * 3 + 0];
-				const uint32 I1 = VertexIndices[TriangleIndex * 3 + 1];
-				const uint32 I2 = VertexIndices[TriangleIndex * 3 + 2];
+	// Build new vertex buffers
+	FMeshBuildVertexData BuildVertexData;
 
-				if (!ensureMsgf(I0 != INDEX_NONE && I1 != INDEX_NONE && I2 != INDEX_NONE,
-					TEXT("Mesh '%s' has triangles with uninitialized vertex indices"), *StaticMesh->GetName()))
-				{
-					continue;
-				}
+	Sections.Empty(PolygonGroups.Num());
 
-				const FVector3f V0 = BuiltVertices[I0];
-				const FVector3f V1 = BuiltVertices[I1];
-				const FVector3f V2 = BuiltVertices[I2];
+	TArray<int32> RemapVerts; //Because we will remove MeshVertex that are redundant, we need a remap
+	//Render data Wedge map is only set for LOD 0???
 
-				const FVector3f TriangleNormal = ((V1 - V2) ^ (V0 - V2));
-				const bool bDegenerateTriangle = TriangleNormal.SizeSquared() < SMALL_NUMBER;
-				if (!bDegenerateTriangle)
-				{
-					BuiltIndices.Add(I0);
-					BuiltIndices.Add(I1);
-					BuiltIndices.Add(I2);
-				}
-			}
-		}
+	TArray<int32> WedgeMap;
+
+	// Prepare the PerSectionIndices array so we can optimize the index buffer for the GPU
+	TArray<TArray<uint32>> PerSectionIndices;
+	PerSectionIndices.AddDefaulted(MeshDescription.PolygonGroups().Num());
+
+	FBoxSphereBounds LODBounds;
+
+	// Build the vertex and index buffer
+	UE::Private::StaticMeshBuilder::BuildVertexBuffer(
+		StaticMesh,
+		MeshDescription,
+		BuildSettings,
+		WedgeMap,
+		Sections,
+		PerSectionIndices,
+		BuildVertexData,
+		MeshDescriptionHelper.GetOverlappingCorners(),
+		RemapVerts,
+		LODBounds,
+		false /* bNeedTangents */,
+		false /* bNeedWedgeMap */
+	);
+
+	BuiltVertices = BuildVertexData.Position;
+
+	// Release MeshDescription memory since we don't need it anymore
+	MeshDescription.Empty();
+
+	// Concatenate the per-section index buffers.
+	bool bNeeds32BitIndices = false;
+	UE::Private::StaticMeshBuilder::BuildCombinedSectionIndices(PerSectionIndices, Sections, BuiltIndices, bNeeds32BitIndices);
+
+	// Apply section remapping
+	for (int32 SectionIndex = 0; SectionIndex < Sections.Num(); SectionIndex++)
+	{
+		Sections[SectionIndex].MaterialIndex = StaticMesh->GetSectionInfoMap().Get(0, SectionIndex).MaterialIndex;
 	}
 
 	return true;
diff --git a/Engine/Source/Developer/MeshBuilder/Public/IMeshBuilderModule.h b/Engine/Source/Developer/MeshBuilder/Public/IMeshBuilderModule.h
--- a/Engine/Source/Developer/MeshBuilder/Public/IMeshBuilderModule.h
+++ b/Engine/Source/Developer/MeshBuilder/Public/IMeshBuilderModule.h
@@ -31,7 +31,7 @@ public:
 
 	virtual bool BuildMesh(class FStaticMeshRenderData& OutRenderData, class UObject* Mesh, const class FStaticMeshLODGroup& LODGroup, bool bGenerateCoarseMeshStreamingLODs, bool bAllowNanite) = 0;
 
-	virtual bool BuildMeshVertexPositions(class UObject* StaticMesh, TArray<uint32>& Indices, TArray<FVector3f>& Vertices) = 0;
+	virtual bool BuildMeshVertexPositions(class UObject* StaticMesh, TArray<uint32>& Indices, TArray<FVector3f>& Vertices, FStaticMeshSectionArray& Sections) = 0;
 
 	virtual bool BuildSkeletalMesh(const struct FSkeletalMeshBuildParameters& SkeletalMeshBuildParameters) = 0;
 
diff --git a/Engine/Source/Developer/MeshBuilder/Public/MeshBuilder.h b/Engine/Source/Developer/MeshBuilder/Public/MeshBuilder.h
--- a/Engine/Source/Developer/MeshBuilder/Public/MeshBuilder.h
+++ b/Engine/Source/Developer/MeshBuilder/Public/MeshBuilder.h
@@ -38,7 +38,8 @@ public:
 	virtual bool BuildMeshVertexPositions(
 		UStaticMesh* StaticMesh,
 		TArray<uint32>& Indices,
-		TArray<FVector3f>& Vertices) = 0;
+		TArray<FVector3f>& Vertices,
+		FStaticMeshSectionArray& Sections) = 0;
 
 	/**
 	 * Build function should be override and is the starting point for skeletal mesh builders
diff --git a/Engine/Source/Developer/MeshBuilder/Public/SkeletalMeshBuilder.h b/Engine/Source/Developer/MeshBuilder/Public/SkeletalMeshBuilder.h
--- a/Engine/Source/Developer/MeshBuilder/Public/SkeletalMeshBuilder.h
+++ b/Engine/Source/Developer/MeshBuilder/Public/SkeletalMeshBuilder.h
@@ -28,7 +28,8 @@ public:
 	virtual bool BuildMeshVertexPositions(
 		UStaticMesh* StaticMesh,
 		TArray<uint32>& Indices,
-		TArray<FVector3f>& Vertices) override
+		TArray<FVector3f>& Vertices,
+		FStaticMeshSectionArray& Sections) override
 	{
 		bool No_Support_For_StaticMesh_Build_In_FSkeletalMeshBuilder_Class = false;
 		check(No_Support_For_StaticMesh_Build_In_FSkeletalMeshBuilder_Class);
diff --git a/Engine/Source/Developer/MeshBuilder/Public/StaticMeshBuilder.h b/Engine/Source/Developer/MeshBuilder/Public/StaticMeshBuilder.h
--- a/Engine/Source/Developer/MeshBuilder/Public/StaticMeshBuilder.h
+++ b/Engine/Source/Developer/MeshBuilder/Public/StaticMeshBuilder.h
@@ -40,7 +40,8 @@ public:
 	virtual bool BuildMeshVertexPositions(
 		UStaticMesh* StaticMesh,
 		TArray<uint32>& Indices,
-		TArray<FVector3f>& Vertices) override;
+		TArray<FVector3f>& Vertices,
+		FStaticMeshSectionArray& Sections) override;
 
 private:
 
diff --git a/Engine/Source/Developer/MeshUtilities/Private/MeshRepresentationCommon.cpp b/Engine/Source/Developer/MeshUtilities/Private/MeshRepresentationCommon.cpp
--- a/Engine/Source/Developer/MeshUtilities/Private/MeshRepresentationCommon.cpp
+++ b/Engine/Source/Developer/MeshUtilities/Private/MeshRepresentationCommon.cpp
@@ -130,11 +130,13 @@ void MeshRepresentation::SetupEmbreeScene(
 	bool bGenerateAsIfTwoSided,
 	FEmbreeScene& EmbreeScene)
 {
+	const uint32 NumVertices = SourceMeshData.IsValid() ? SourceMeshData.GetNumVertices() : LODModel.VertexBuffers.PositionVertexBuffer.GetNumVertices();
 	const uint32 NumIndices = SourceMeshData.IsValid() ? SourceMeshData.GetNumIndices() : LODModel.IndexBuffer.GetNumIndices();
 	const int32 NumTriangles = NumIndices / 3;
-	const uint32 NumVertices = SourceMeshData.IsValid() ? SourceMeshData.GetNumVertices() : LODModel.VertexBuffers.PositionVertexBuffer.GetNumVertices();
 	EmbreeScene.NumIndices = NumTriangles;
 
+	const FStaticMeshSectionArray& Sections = SourceMeshData.IsValid() ? SourceMeshData.Sections : LODModel.Sections;
+
 	TArray<FkDOPBuildCollisionTriangle<uint32> > BuildTriangles;
 
 #if USE_EMBREE
@@ -175,65 +177,59 @@ void MeshRepresentation::SetupEmbreeScene(
 	TArray<int32> FilteredTriangles;
 	FilteredTriangles.Empty(NumTriangles);
 
-	if (SourceMeshData.IsValid())
+	for (int32 TriangleIndex = 0; TriangleIndex < NumTriangles; ++TriangleIndex)
 	{
-		for (int32 TriangleIndex = 0; TriangleIndex < NumTriangles; ++TriangleIndex)
+		FVector3f V0, V1, V2;
+
+		if (SourceMeshData.IsValid())
 		{
 			const uint32 I0 = SourceMeshData.TriangleIndices[TriangleIndex * 3 + 0];
 			const uint32 I1 = SourceMeshData.TriangleIndices[TriangleIndex * 3 + 1];
 			const uint32 I2 = SourceMeshData.TriangleIndices[TriangleIndex * 3 + 2];
 
-			const FVector3f V0 = SourceMeshData.VertexPositions[I0];
-			const FVector3f V1 = SourceMeshData.VertexPositions[I1];
-			const FVector3f V2 = SourceMeshData.VertexPositions[I2];
-
-			const FVector3f TriangleNormal = ((V1 - V2) ^ (V0 - V2));
-			const bool bDegenerateTriangle = TriangleNormal.SizeSquared() < SMALL_NUMBER;
-			if (!bDegenerateTriangle)
-			{
-				FilteredTriangles.Add(TriangleIndex);
-			}
+			V0 = SourceMeshData.VertexPositions[I0];
+			V1 = SourceMeshData.VertexPositions[I1];
+			V2 = SourceMeshData.VertexPositions[I2];
 		}
-	}
-	else
-	{
-		for (int32 TriangleIndex = 0; TriangleIndex < NumTriangles; ++TriangleIndex)
+		else
 		{
 			const FIndexArrayView Indices = LODModel.IndexBuffer.GetArrayView();
 			const uint32 I0 = Indices[TriangleIndex * 3 + 0];
 			const uint32 I1 = Indices[TriangleIndex * 3 + 1];
 			const uint32 I2 = Indices[TriangleIndex * 3 + 2];
 
-			const FVector3f V0 = LODModel.VertexBuffers.PositionVertexBuffer.VertexPosition(I0);
-			const FVector3f V1 = LODModel.VertexBuffers.PositionVertexBuffer.VertexPosition(I1);
-			const FVector3f V2 = LODModel.VertexBuffers.PositionVertexBuffer.VertexPosition(I2);
+			V0 = LODModel.VertexBuffers.PositionVertexBuffer.VertexPosition(I0);
+			V1 = LODModel.VertexBuffers.PositionVertexBuffer.VertexPosition(I1);
+			V2 = LODModel.VertexBuffers.PositionVertexBuffer.VertexPosition(I2);
+		}
+
+		const FVector3f TriangleNormal = ((V1 - V2) ^ (V0 - V2));
+		const bool bDegenerateTriangle = TriangleNormal.SizeSquared() < SMALL_NUMBER;
+		if (!bDegenerateTriangle)
+		{
+			//FilteredTriangles.Add(TriangleIndex);
+
+			bool bTriangleIsOpaqueOrMasked = false;
 
-			const FVector3f TriangleNormal = ((V1 - V2) ^ (V0 - V2));
-			const bool bDegenerateTriangle = TriangleNormal.SizeSquared() < SMALL_NUMBER;
-			if (!bDegenerateTriangle)
+			for (int32 SectionIndex = 0; SectionIndex < Sections.Num(); SectionIndex++)
 			{
-				bool bTriangleIsOpaqueOrMasked = false;
+				const FStaticMeshSection& Section = Sections[SectionIndex];
 
-				for (int32 SectionIndex = 0; SectionIndex < LODModel.Sections.Num(); SectionIndex++)
+				if ((uint32)(TriangleIndex * 3) >= Section.FirstIndex && (uint32)(TriangleIndex * 3) < Section.FirstIndex + Section.NumTriangles * 3)
 				{
-					const FStaticMeshSection& Section = LODModel.Sections[SectionIndex];
-
-					if ((uint32)(TriangleIndex * 3) >= Section.FirstIndex && (uint32)(TriangleIndex * 3) < Section.FirstIndex + Section.NumTriangles * 3)
+					if (SectionData.IsValidIndex(SectionIndex))
 					{
-						if (SectionData.IsValidIndex(SectionIndex))
-						{
-							bTriangleIsOpaqueOrMasked = !IsTranslucentBlendMode(SectionData[SectionIndex].BlendMode) && SectionData[SectionIndex].bAffectDistanceFieldLighting;
-						}
-
-						break;
+						bTriangleIsOpaqueOrMasked = !IsTranslucentBlendMode(SectionData[SectionIndex].BlendMode) && SectionData[SectionIndex].bAffectDistanceFieldLighting;
 					}
-				}
 
-				if (bTriangleIsOpaqueOrMasked)
-				{
-					FilteredTriangles.Add(TriangleIndex);
+					break;
 				}
 			}
+
+			if (bTriangleIsOpaqueOrMasked)
+			{
+				FilteredTriangles.Add(TriangleIndex);
+			}
 		}
 	}
 
@@ -280,9 +276,9 @@ void MeshRepresentation::SetupEmbreeScene(
 
 		bool bTriangleIsTwoSided = false;
 
-		for (int32 SectionIndex = 0; SectionIndex < LODModel.Sections.Num(); SectionIndex++)
+		for (int32 SectionIndex = 0; SectionIndex < Sections.Num(); SectionIndex++)
 		{
-			const FStaticMeshSection& Section = LODModel.Sections[SectionIndex];
+			const FStaticMeshSection& Section = Sections[SectionIndex];
 
 			if ((uint32)(TriangleIndex * 3) >= Section.FirstIndex && (uint32)(TriangleIndex * 3) < Section.FirstIndex + Section.NumTriangles * 3)
 			{
@@ -357,9 +353,9 @@ void MeshRepresentation::SetupEmbreeScene(
 		uint32 NumTrianglesTotal = 0;
 		uint32 NumTwoSidedTriangles = 0;
 
-		for (int32 SectionIndex = 0; SectionIndex < LODModel.Sections.Num(); SectionIndex++)
+		for (int32 SectionIndex = 0; SectionIndex < Sections.Num(); SectionIndex++)
 		{
-			const FStaticMeshSection& Section = LODModel.Sections[SectionIndex];
+			const FStaticMeshSection& Section = Sections[SectionIndex];
 
 			if (SectionData.IsValidIndex(SectionIndex))
 			{
diff --git a/Engine/Source/Runtime/Engine/Private/DistanceFieldAtlas.cpp b/Engine/Source/Runtime/Engine/Private/DistanceFieldAtlas.cpp
--- a/Engine/Source/Runtime/Engine/Private/DistanceFieldAtlas.cpp
+++ b/Engine/Source/Runtime/Engine/Private/DistanceFieldAtlas.cpp
@@ -125,7 +125,7 @@ FDistanceFieldAsyncQueue* GDistanceFieldAsyncQueue = NULL;
 #if WITH_EDITOR
 
 // DDC key for distance field data, must be changed when modifying the generation code or data format
-#define DISTANCEFIELD_DERIVEDDATA_VER TEXT("C8D7BA5F-5874-4FFB-A264-7B7DE9299F68")
+#define DISTANCEFIELD_DERIVEDDATA_VER TEXT("E74258BD-3117-4B00-B97C-29DDA01A3D75")
 
 FString BuildDistanceFieldDerivedDataKey(const FString& InMeshKey)
 {
@@ -228,7 +228,7 @@ void FDistanceFieldVolumeData::CacheDerivedData(const FString& InStaticMeshDeriv
 		if (GenerateSource->IsNaniteEnabled())
 		{
 			IMeshBuilderModule& MeshBuilderModule = IMeshBuilderModule::GetForPlatform(TargetPlatform);
-			if (!MeshBuilderModule.BuildMeshVertexPositions(GenerateSource, NewTask->SourceMeshData.TriangleIndices, NewTask->SourceMeshData.VertexPositions))
+			if (!MeshBuilderModule.BuildMeshVertexPositions(GenerateSource, NewTask->SourceMeshData.TriangleIndices, NewTask->SourceMeshData.VertexPositions, NewTask->SourceMeshData.Sections))
 			{
 				UE_LOG(LogStaticMesh, Error, TEXT("Failed to build static mesh. See previous line(s) for details."));
 			}
diff --git a/Engine/Source/Runtime/Engine/Private/MeshCardRepresentation.cpp b/Engine/Source/Runtime/Engine/Private/MeshCardRepresentation.cpp
--- a/Engine/Source/Runtime/Engine/Private/MeshCardRepresentation.cpp
+++ b/Engine/Source/Runtime/Engine/Private/MeshCardRepresentation.cpp
@@ -98,7 +98,7 @@ FCardRepresentationAsyncQueue* GCardRepresentationAsyncQueue = NULL;
 #if WITH_EDITOR
 
 // DDC key for card representation data, must be changed when modifying the generation code or data format
-#define CARDREPRESENTATION_DERIVEDDATA_VER TEXT("E8B65E19-7C68-4BDC-BA5E-91C9A5855599")
+#define CARDREPRESENTATION_DERIVEDDATA_VER TEXT("5A576281-130D-4159-B20C-9A9880228C88")
 
 FString BuildCardRepresentationDerivedDataKey(const FString& InMeshKey, int32 MaxLumenMeshCards)
 {
@@ -187,7 +187,7 @@ void FCardRepresentationData::CacheDerivedData(const FString& InDDCKey, const IT
 		else if (Mesh->IsNaniteEnabled())
 		{
 			IMeshBuilderModule& MeshBuilderModule = IMeshBuilderModule::GetForPlatform(TargetPlatform);
-			if (!MeshBuilderModule.BuildMeshVertexPositions(Mesh, NewTask->SourceMeshData.TriangleIndices, NewTask->SourceMeshData.VertexPositions))
+			if (!MeshBuilderModule.BuildMeshVertexPositions(Mesh, NewTask->SourceMeshData.TriangleIndices, NewTask->SourceMeshData.VertexPositions, NewTask->SourceMeshData.Sections))
 			{
 				UE_LOG(LogStaticMesh, Error, TEXT("Failed to build static mesh. See previous line(s) for details."));
 			}
diff --git a/Engine/Source/Runtime/Engine/Public/DerivedMeshDataTaskUtils.h b/Engine/Source/Runtime/Engine/Public/DerivedMeshDataTaskUtils.h
--- a/Engine/Source/Runtime/Engine/Public/DerivedMeshDataTaskUtils.h
+++ b/Engine/Source/Runtime/Engine/Public/DerivedMeshDataTaskUtils.h
@@ -15,6 +15,7 @@ class FSourceMeshDataForDerivedDataTask
 public:
 	TArray<uint32> TriangleIndices;
 	TArray<FVector3f> VertexPositions;
+	FStaticMeshSectionArray Sections;
 
 	int32 GetNumIndices() const
 	{
