[Backout] - CL26160360, 26157603, 26146080, 26137852, 26126752
Original CL Desc
-----------------------------------------------------------------
Use FSignedDistanceFieldBuildSectionData when generating distance fields for Nanite meshes.

Context:
- Nanite meshes distance fields are generated using the original mesh data instead of the coarse representation generated by Nanite, however that code path was not checking section data such as blend mode, two sided, etc.

Fix:
- Modified FStaticMeshBuilder::BuildMeshVertexPositions(...) to also output a FStaticMeshSectionArray.
- Updated MeshRepresentation::SetupEmbreeScene(...) to check section flags when using SourceMeshData.IsValid() code path.

#prelight

diff --git a/Engine/Source/Developer/MeshBuilder/Private/MeshBuilderModule.cpp b/Engine/Source/Developer/MeshBuilder/Private/MeshBuilderModule.cpp
--- a/Engine/Source/Developer/MeshBuilder/Private/MeshBuilderModule.cpp
+++ b/Engine/Source/Developer/MeshBuilder/Private/MeshBuilderModule.cpp
@@ -31,8 +31,7 @@ public:
 	virtual bool BuildMeshVertexPositions(
 		UObject* StaticMesh,
 		TArray<uint32>& Indices,
-		TArray<FVector3f>& Vertices,
-		FStaticMeshSectionArray& Sections) override;
+		TArray<FVector3f>& Vertices) override;
 
 	virtual bool BuildSkeletalMesh(const FSkeletalMeshBuildParameters& SkeletalMeshBuildParameters) override;
 
@@ -56,14 +55,13 @@ bool FMeshBuilderModule::BuildMesh(FStaticMeshRenderData& OutRenderData, class U
 bool FMeshBuilderModule::BuildMeshVertexPositions(
 	UObject* Mesh,
 	TArray<uint32>& Indices,
-	TArray<FVector3f>& Vertices,
-	FStaticMeshSectionArray& Sections)
+	TArray<FVector3f>& Vertices)
 {
 	UStaticMesh* StaticMesh = Cast< UStaticMesh >(Mesh);
 	if (StaticMesh)
 	{
 		//Call the static mesh builder
-		return FStaticMeshBuilder().BuildMeshVertexPositions(StaticMesh, Indices, Vertices, Sections);
+		return FStaticMeshBuilder().BuildMeshVertexPositions(StaticMesh, Indices, Vertices);
 	}
 	return false;
 }
diff --git a/Engine/Source/Developer/MeshBuilder/Private/StaticMeshBuilder.cpp b/Engine/Source/Developer/MeshBuilder/Private/StaticMeshBuilder.cpp
--- a/Engine/Source/Developer/MeshBuilder/Private/StaticMeshBuilder.cpp
+++ b/Engine/Source/Developer/MeshBuilder/Private/StaticMeshBuilder.cpp
@@ -730,8 +730,7 @@ bool FStaticMeshBuilder::Build(FStaticMeshRenderData& StaticMeshRenderData, USta
 bool FStaticMeshBuilder::BuildMeshVertexPositions(
 	UStaticMesh* StaticMesh,
 	TArray<uint32>& BuiltIndices,
-	TArray<FVector3f>& BuiltVertices,
-	FStaticMeshSectionArray& Sections)
+	TArray<FVector3f>& BuiltVertices)
 {
 	TRACE_CPUPROFILER_EVENT_SCOPE(FStaticMeshBuilder::BuildMeshVertexPositions);
 
@@ -743,62 +742,59 @@ bool FStaticMeshBuilder::BuildMeshVertexPositions(
 		return false;
 	}
 
-	FMeshDescription MeshDescription;
-	const bool bIsMeshDescriptionValid = SourceModel.CloneMeshDescription(MeshDescription);
-	check(bIsMeshDescriptionValid);
-
-	FMeshBuildSettings& BuildSettings = StaticMesh->GetSourceModel(0).BuildSettings;
-
-	FMeshDescriptionHelper MeshDescriptionHelper(&BuildSettings);
-	MeshDescriptionHelper.SetupRenderMeshDescription(StaticMesh, MeshDescription, false, false);
-
-	const FPolygonGroupArray& PolygonGroups = MeshDescription.PolygonGroups();
-
-	// Build new vertex buffers
-	FMeshBuildVertexData BuildVertexData;
-
-	Sections.Empty(PolygonGroups.Num());
-
-	TArray<int32> RemapVerts; //Because we will remove MeshVertex that are redundant, we need a remap
-	//Render data Wedge map is only set for LOD 0???
-
-	TArray<int32> WedgeMap;
+	const int32 NumSourceModels = StaticMesh->GetNumSourceModels();
+	if (NumSourceModels > 0)
+	{
+		FMeshDescription MeshDescription;
+		const bool bIsMeshDescriptionValid = SourceModel.CloneMeshDescription(MeshDescription);
+		if (bIsMeshDescriptionValid)
+		{
+			FElementIDRemappings Remappings;
+			if (MeshDescription.NeedsCompact())
+			{
+				MeshDescription.Compact(Remappings);
+			}
 
-	// Prepare the PerSectionIndices array so we can optimize the index buffer for the GPU
-	TArray<TArray<uint32>> PerSectionIndices;
-	PerSectionIndices.AddDefaulted(MeshDescription.PolygonGroups().Num());
+			const FMeshBuildSettings& BuildSettings = SourceModel.BuildSettings;
 
-	FBoxSphereBounds LODBounds;
+			const FStaticMeshConstAttributes Attributes(MeshDescription);
+			TArrayView<const FVector3f> VertexPositions = Attributes.GetVertexPositions().GetRawArray();
+			TArrayView<const FVertexID> VertexIndices = Attributes.GetTriangleVertexIndices().GetRawArray();
+			const FVector3f BuildScale3D = (FVector3f)BuildSettings.BuildScale3D;
 
-	// Build the vertex and index buffer
-	UE::Private::StaticMeshBuilder::BuildVertexBuffer(
-		StaticMesh,
-		MeshDescription,
-		BuildSettings,
-		WedgeMap,
-		Sections,
-		PerSectionIndices,
-		BuildVertexData,
-		MeshDescriptionHelper.GetOverlappingCorners(),
-		RemapVerts,
-		LODBounds,
-		false /* bNeedTangents */,
-		false /* bNeedWedgeMap */
-	);
+			BuiltVertices.Reserve(VertexPositions.Num());
+			for (int32 VertexIndex = 0; VertexIndex < VertexPositions.Num(); ++VertexIndex)
+			{
+				BuiltVertices.Add(VertexPositions[VertexIndex] * BuildScale3D);
+			}
 
-	BuiltVertices = BuildVertexData.Position;
+			BuiltIndices.Reserve(VertexIndices.Num());
+			for (int32 TriangleIndex = 0; TriangleIndex < VertexIndices.Num() / 3; ++TriangleIndex)
+			{
+				const uint32 I0 = VertexIndices[TriangleIndex * 3 + 0];
+				const uint32 I1 = VertexIndices[TriangleIndex * 3 + 1];
+				const uint32 I2 = VertexIndices[TriangleIndex * 3 + 2];
 
-	// Release MeshDescription memory since we don't need it anymore
-	MeshDescription.Empty();
+				if (!ensureMsgf(I0 != INDEX_NONE && I1 != INDEX_NONE && I2 != INDEX_NONE,
+					TEXT("Mesh '%s' has triangles with uninitialized vertex indices"), *StaticMesh->GetName()))
+				{
+					continue;
+				}
 
-	// Concatenate the per-section index buffers.
-	bool bNeeds32BitIndices = false;
-	UE::Private::StaticMeshBuilder::BuildCombinedSectionIndices(PerSectionIndices, Sections, BuiltIndices, bNeeds32BitIndices);
+				const FVector3f V0 = BuiltVertices[I0];
+				const FVector3f V1 = BuiltVertices[I1];
+				const FVector3f V2 = BuiltVertices[I2];
 
-	// Apply section remapping
-	for (int32 SectionIndex = 0; SectionIndex < Sections.Num(); SectionIndex++)
-	{
-		Sections[SectionIndex].MaterialIndex = StaticMesh->GetSectionInfoMap().Get(0, SectionIndex).MaterialIndex;
+				const FVector3f TriangleNormal = ((V1 - V2) ^ (V0 - V2));
+				const bool bDegenerateTriangle = TriangleNormal.SizeSquared() < SMALL_NUMBER;
+				if (!bDegenerateTriangle)
+				{
+					BuiltIndices.Add(I0);
+					BuiltIndices.Add(I1);
+					BuiltIndices.Add(I2);
+				}
+			}
+		}
 	}
 
 	return true;
diff --git a/Engine/Source/Developer/MeshBuilder/Public/IMeshBuilderModule.h b/Engine/Source/Developer/MeshBuilder/Public/IMeshBuilderModule.h
--- a/Engine/Source/Developer/MeshBuilder/Public/IMeshBuilderModule.h
+++ b/Engine/Source/Developer/MeshBuilder/Public/IMeshBuilderModule.h
@@ -31,7 +31,7 @@ public:
 
 	virtual bool BuildMesh(class FStaticMeshRenderData& OutRenderData, class UObject* Mesh, const class FStaticMeshLODGroup& LODGroup, bool bGenerateCoarseMeshStreamingLODs, bool bAllowNanite) = 0;
 
-	virtual bool BuildMeshVertexPositions(class UObject* StaticMesh, TArray<uint32>& Indices, TArray<FVector3f>& Vertices, FStaticMeshSectionArray& Sections) = 0;
+	virtual bool BuildMeshVertexPositions(class UObject* StaticMesh, TArray<uint32>& Indices, TArray<FVector3f>& Vertices) = 0;
 
 	virtual bool BuildSkeletalMesh(const struct FSkeletalMeshBuildParameters& SkeletalMeshBuildParameters) = 0;
 
diff --git a/Engine/Source/Developer/MeshBuilder/Public/MeshBuilder.h b/Engine/Source/Developer/MeshBuilder/Public/MeshBuilder.h
--- a/Engine/Source/Developer/MeshBuilder/Public/MeshBuilder.h
+++ b/Engine/Source/Developer/MeshBuilder/Public/MeshBuilder.h
@@ -38,8 +38,7 @@ public:
 	virtual bool BuildMeshVertexPositions(
 		UStaticMesh* StaticMesh,
 		TArray<uint32>& Indices,
-		TArray<FVector3f>& Vertices,
-		FStaticMeshSectionArray& Sections) = 0;
+		TArray<FVector3f>& Vertices) = 0;
 
 	/**
 	 * Build function should be override and is the starting point for skeletal mesh builders
diff --git a/Engine/Source/Developer/MeshBuilder/Public/SkeletalMeshBuilder.h b/Engine/Source/Developer/MeshBuilder/Public/SkeletalMeshBuilder.h
--- a/Engine/Source/Developer/MeshBuilder/Public/SkeletalMeshBuilder.h
+++ b/Engine/Source/Developer/MeshBuilder/Public/SkeletalMeshBuilder.h
@@ -28,8 +28,7 @@ public:
 	virtual bool BuildMeshVertexPositions(
 		UStaticMesh* StaticMesh,
 		TArray<uint32>& Indices,
-		TArray<FVector3f>& Vertices,
-		FStaticMeshSectionArray& Sections) override
+		TArray<FVector3f>& Vertices) override
 	{
 		bool No_Support_For_StaticMesh_Build_In_FSkeletalMeshBuilder_Class = false;
 		check(No_Support_For_StaticMesh_Build_In_FSkeletalMeshBuilder_Class);
diff --git a/Engine/Source/Developer/MeshBuilder/Public/StaticMeshBuilder.h b/Engine/Source/Developer/MeshBuilder/Public/StaticMeshBuilder.h
--- a/Engine/Source/Developer/MeshBuilder/Public/StaticMeshBuilder.h
+++ b/Engine/Source/Developer/MeshBuilder/Public/StaticMeshBuilder.h
@@ -40,8 +40,7 @@ public:
 	virtual bool BuildMeshVertexPositions(
 		UStaticMesh* StaticMesh,
 		TArray<uint32>& Indices,
-		TArray<FVector3f>& Vertices,
-		FStaticMeshSectionArray& Sections) override;
+		TArray<FVector3f>& Vertices) override;
 
 private:
 
diff --git a/Engine/Source/Developer/MeshUtilities/Private/MeshCardRepresentationUtilities.cpp b/Engine/Source/Developer/MeshUtilities/Private/MeshCardRepresentationUtilities.cpp
--- a/Engine/Source/Developer/MeshUtilities/Private/MeshCardRepresentationUtilities.cpp
+++ b/Engine/Source/Developer/MeshUtilities/Private/MeshCardRepresentationUtilities.cpp
@@ -1054,7 +1054,7 @@ bool FMeshUtilities::GenerateCardRepresentationData(
 	const FSourceMeshDataForDerivedDataTask& SourceMeshData,
 	const FStaticMeshLODResources& LODModel,
 	class FQueuedThreadPool& ThreadPool,
-	const TArray<FSignedDistanceFieldBuildSectionData>& SectionData,
+	const TArray<FSignedDistanceFieldBuildMaterialData>& MaterialBlendModes,
 	const FBoxSphereBounds& Bounds,
 	const FDistanceFieldVolumeData* DistanceFieldVolumeData,
 	int32 MaxLumenMeshCards,
@@ -1072,7 +1072,7 @@ bool FMeshUtilities::GenerateCardRepresentationData(
 		MeshRepresentation::SetupEmbreeScene(MeshName,
 			SourceMeshData,
 			LODModel,
-			SectionData,
+			MaterialBlendModes,
 			bGenerateAsIfTwoSided,
 			EmbreeScene);
 
diff --git a/Engine/Source/Developer/MeshUtilities/Private/MeshDistanceFieldUtilities.cpp b/Engine/Source/Developer/MeshUtilities/Private/MeshDistanceFieldUtilities.cpp
--- a/Engine/Source/Developer/MeshUtilities/Private/MeshDistanceFieldUtilities.cpp
+++ b/Engine/Source/Developer/MeshUtilities/Private/MeshDistanceFieldUtilities.cpp
@@ -252,7 +252,7 @@ void FMeshUtilities::GenerateSignedDistanceFieldVolumeData(
 	const FSourceMeshDataForDerivedDataTask& SourceMeshData,
 	const FStaticMeshLODResources& LODModel,
 	class FQueuedThreadPool& ThreadPool,
-	const TArray<FSignedDistanceFieldBuildSectionData>& SectionData,
+	const TArray<FSignedDistanceFieldBuildMaterialData>& MaterialBlendModes,
 	const FBoxSphereBounds& Bounds,
 	float DistanceFieldResolutionScale,
 	bool bGenerateAsIfTwoSided,
@@ -268,7 +268,7 @@ void FMeshUtilities::GenerateSignedDistanceFieldVolumeData(
 		MeshRepresentation::SetupEmbreeScene(MeshName,
 			SourceMeshData,
 			LODModel,
-			SectionData,
+			MaterialBlendModes,
 			bGenerateAsIfTwoSided,
 			EmbreeScene);
 
@@ -529,7 +529,7 @@ void FMeshUtilities::GenerateSignedDistanceFieldVolumeData(
 	const FSourceMeshDataForDerivedDataTask& SourceMeshData,
 	const FStaticMeshLODResources& LODModel,
 	class FQueuedThreadPool& ThreadPool,
-	const TArray<FSignedDistanceFieldBuildSectionData>& SectionData,
+	const TArray<FSignedDistanceFieldBuildMaterialData>& MaterialBlendModes,
 	const FBoxSphereBounds& Bounds,
 	float DistanceFieldResolutionScale,
 	bool bGenerateAsIfTwoSided,
diff --git a/Engine/Source/Developer/MeshUtilities/Private/MeshRepresentationCommon.cpp b/Engine/Source/Developer/MeshUtilities/Private/MeshRepresentationCommon.cpp
--- a/Engine/Source/Developer/MeshUtilities/Private/MeshRepresentationCommon.cpp
+++ b/Engine/Source/Developer/MeshUtilities/Private/MeshRepresentationCommon.cpp
@@ -126,17 +126,15 @@ void MeshRepresentation::SetupEmbreeScene(
 	FString MeshName,
 	const FSourceMeshDataForDerivedDataTask& SourceMeshData,
 	const FStaticMeshLODResources& LODModel,
-	const TArray<FSignedDistanceFieldBuildSectionData>& SectionData,
+	const TArray<FSignedDistanceFieldBuildMaterialData>& MaterialBlendModes,
 	bool bGenerateAsIfTwoSided,
 	FEmbreeScene& EmbreeScene)
 {
-	const uint32 NumVertices = SourceMeshData.IsValid() ? SourceMeshData.GetNumVertices() : LODModel.VertexBuffers.PositionVertexBuffer.GetNumVertices();
 	const uint32 NumIndices = SourceMeshData.IsValid() ? SourceMeshData.GetNumIndices() : LODModel.IndexBuffer.GetNumIndices();
 	const int32 NumTriangles = NumIndices / 3;
+	const uint32 NumVertices = SourceMeshData.IsValid() ? SourceMeshData.GetNumVertices() : LODModel.VertexBuffers.PositionVertexBuffer.GetNumVertices();
 	EmbreeScene.NumIndices = NumTriangles;
 
-	const FStaticMeshSectionArray& Sections = SourceMeshData.IsValid() ? SourceMeshData.Sections : LODModel.Sections;
-
 	TArray<FkDOPBuildCollisionTriangle<uint32> > BuildTriangles;
 
 #if USE_EMBREE
@@ -167,68 +165,67 @@ void MeshRepresentation::SetupEmbreeScene(
 	}
 #endif
 
-	/*
-	if (LODModel.Sections.Num() > SectionData.Num())
-	{
-		UE_LOG(LogMeshUtilities, Warning, TEXT("Unexpected number of mesh sections when setting up Embree Scene for %s."), *MeshName);
-	}
-	*/
-
 	TArray<int32> FilteredTriangles;
 	FilteredTriangles.Empty(NumTriangles);
 
-	for (int32 TriangleIndex = 0; TriangleIndex < NumTriangles; ++TriangleIndex)
+	if (SourceMeshData.IsValid())
 	{
-		FVector3f V0, V1, V2;
-
-		if (SourceMeshData.IsValid())
+		for (int32 TriangleIndex = 0; TriangleIndex < NumTriangles; ++TriangleIndex)
 		{
 			const uint32 I0 = SourceMeshData.TriangleIndices[TriangleIndex * 3 + 0];
 			const uint32 I1 = SourceMeshData.TriangleIndices[TriangleIndex * 3 + 1];
 			const uint32 I2 = SourceMeshData.TriangleIndices[TriangleIndex * 3 + 2];
 
-			V0 = SourceMeshData.VertexPositions[I0];
-			V1 = SourceMeshData.VertexPositions[I1];
-			V2 = SourceMeshData.VertexPositions[I2];
+			const FVector3f V0 = SourceMeshData.VertexPositions[I0];
+			const FVector3f V1 = SourceMeshData.VertexPositions[I1];
+			const FVector3f V2 = SourceMeshData.VertexPositions[I2];
+
+			const FVector3f TriangleNormal = ((V1 - V2) ^ (V0 - V2));
+			const bool bDegenerateTriangle = TriangleNormal.SizeSquared() < SMALL_NUMBER;
+			if (!bDegenerateTriangle)
+			{
+				FilteredTriangles.Add(TriangleIndex);
+			}
 		}
-		else
+	}
+	else
+	{
+		for (int32 TriangleIndex = 0; TriangleIndex < NumTriangles; ++TriangleIndex)
 		{
 			const FIndexArrayView Indices = LODModel.IndexBuffer.GetArrayView();
 			const uint32 I0 = Indices[TriangleIndex * 3 + 0];
 			const uint32 I1 = Indices[TriangleIndex * 3 + 1];
 			const uint32 I2 = Indices[TriangleIndex * 3 + 2];
 
-			V0 = LODModel.VertexBuffers.PositionVertexBuffer.VertexPosition(I0);
-			V1 = LODModel.VertexBuffers.PositionVertexBuffer.VertexPosition(I1);
-			V2 = LODModel.VertexBuffers.PositionVertexBuffer.VertexPosition(I2);
-		}
-
-		const FVector3f TriangleNormal = ((V1 - V2) ^ (V0 - V2));
-		const bool bDegenerateTriangle = TriangleNormal.SizeSquared() < SMALL_NUMBER;
-		if (!bDegenerateTriangle)
-		{
-			//FilteredTriangles.Add(TriangleIndex);
-
-			bool bTriangleIsOpaqueOrMasked = false;
+			const FVector3f V0 = LODModel.VertexBuffers.PositionVertexBuffer.VertexPosition(I0);
+			const FVector3f V1 = LODModel.VertexBuffers.PositionVertexBuffer.VertexPosition(I1);
+			const FVector3f V2 = LODModel.VertexBuffers.PositionVertexBuffer.VertexPosition(I2);
 
-			for (int32 SectionIndex = 0; SectionIndex < Sections.Num(); SectionIndex++)
+			const FVector3f TriangleNormal = ((V1 - V2) ^ (V0 - V2));
+			const bool bDegenerateTriangle = TriangleNormal.SizeSquared() < SMALL_NUMBER;
+			if (!bDegenerateTriangle)
 			{
-				const FStaticMeshSection& Section = Sections[SectionIndex];
+				bool bTriangleIsOpaqueOrMasked = false;
 
-				if ((uint32)(TriangleIndex * 3) >= Section.FirstIndex && (uint32)(TriangleIndex * 3) < Section.FirstIndex + Section.NumTriangles * 3)
+				for (int32 SectionIndex = 0; SectionIndex < LODModel.Sections.Num(); SectionIndex++)
 				{
-					if (SectionData.IsValidIndex(SectionIndex))
+					const FStaticMeshSection& Section = LODModel.Sections[SectionIndex];
+
+					if ((uint32)(TriangleIndex * 3) >= Section.FirstIndex && (uint32)(TriangleIndex * 3) < Section.FirstIndex + Section.NumTriangles * 3)
 					{
-						bTriangleIsOpaqueOrMasked = !IsTranslucentBlendMode(SectionData[SectionIndex].BlendMode) && SectionData[SectionIndex].bAffectDistanceFieldLighting;
-					}
+						if (MaterialBlendModes.IsValidIndex(Section.MaterialIndex))
+						{
+							bTriangleIsOpaqueOrMasked = !IsTranslucentBlendMode(MaterialBlendModes[Section.MaterialIndex].BlendMode) && MaterialBlendModes[Section.MaterialIndex].bAffectDistanceFieldLighting;
+						}
 
-					break;
+						break;
+					}
 				}
-			}
 
-			if (bTriangleIsOpaqueOrMasked)
-			{
-				FilteredTriangles.Add(TriangleIndex);
+				if (bTriangleIsOpaqueOrMasked)
+				{
+					FilteredTriangles.Add(TriangleIndex);
+				}
 			}
 		}
 	}
@@ -276,15 +273,15 @@ void MeshRepresentation::SetupEmbreeScene(
 
 		bool bTriangleIsTwoSided = false;
 
-		for (int32 SectionIndex = 0; SectionIndex < Sections.Num(); SectionIndex++)
+		for (int32 SectionIndex = 0; SectionIndex < LODModel.Sections.Num(); SectionIndex++)
 		{
-			const FStaticMeshSection& Section = Sections[SectionIndex];
+			const FStaticMeshSection& Section = LODModel.Sections[SectionIndex];
 
 			if ((uint32)(TriangleIndex * 3) >= Section.FirstIndex && (uint32)(TriangleIndex * 3) < Section.FirstIndex + Section.NumTriangles * 3)
 			{
-				if (SectionData.IsValidIndex(SectionIndex))
+				if (MaterialBlendModes.IsValidIndex(Section.MaterialIndex))
 				{
-					bTriangleIsTwoSided = SectionData[SectionIndex].bTwoSided;
+					bTriangleIsTwoSided = MaterialBlendModes[Section.MaterialIndex].bTwoSided;
 				}
 
 				break;
@@ -353,15 +350,15 @@ void MeshRepresentation::SetupEmbreeScene(
 		uint32 NumTrianglesTotal = 0;
 		uint32 NumTwoSidedTriangles = 0;
 
-		for (int32 SectionIndex = 0; SectionIndex < Sections.Num(); SectionIndex++)
+		for (int32 SectionIndex = 0; SectionIndex < LODModel.Sections.Num(); SectionIndex++)
 		{
-			const FStaticMeshSection& Section = Sections[SectionIndex];
+			const FStaticMeshSection& Section = LODModel.Sections[SectionIndex];
 
-			if (SectionData.IsValidIndex(SectionIndex))
+			if (MaterialBlendModes.IsValidIndex(Section.MaterialIndex))
 			{
 				NumTrianglesTotal += Section.NumTriangles;
 
-				if (SectionData[SectionIndex].bTwoSided)
+				if (MaterialBlendModes[Section.MaterialIndex].bTwoSided)
 				{
 					NumTwoSidedTriangles += Section.NumTriangles;
 				}
diff --git a/Engine/Source/Developer/MeshUtilities/Private/MeshRepresentationCommon.h b/Engine/Source/Developer/MeshUtilities/Private/MeshRepresentationCommon.h
--- a/Engine/Source/Developer/MeshUtilities/Private/MeshRepresentationCommon.h
+++ b/Engine/Source/Developer/MeshUtilities/Private/MeshRepresentationCommon.h
@@ -159,7 +159,7 @@ namespace MeshRepresentation
 	void SetupEmbreeScene(FString MeshName,
 		const FSourceMeshDataForDerivedDataTask& SourceMeshData,
 		const FStaticMeshLODResources& LODModel,
-		const TArray<FSignedDistanceFieldBuildSectionData>& SectionData,
+		const TArray<FSignedDistanceFieldBuildMaterialData>& MaterialBlendModes,
 		bool bGenerateAsIfTwoSided,
 		FEmbreeScene& EmbreeScene);
 
diff --git a/Engine/Source/Developer/MeshUtilities/Private/MeshUtilitiesPrivate.h b/Engine/Source/Developer/MeshUtilities/Private/MeshUtilitiesPrivate.h
--- a/Engine/Source/Developer/MeshUtilities/Private/MeshUtilitiesPrivate.h
+++ b/Engine/Source/Developer/MeshUtilities/Private/MeshUtilitiesPrivate.h
@@ -60,7 +60,7 @@ private:
 		const FSourceMeshDataForDerivedDataTask& SourceMeshData,
 		const FStaticMeshLODResources& LODModel,
 		class FQueuedThreadPool& ThreadPool,
-		const TArray<FSignedDistanceFieldBuildSectionData>& SectionData,
+		const TArray<FSignedDistanceFieldBuildMaterialData>& MaterialBlendModes,
 		const FBoxSphereBounds& Bounds,
 		float DistanceFieldResolutionScale,
 		bool bGenerateAsIfTwoSided,
@@ -71,7 +71,7 @@ private:
 		const FSourceMeshDataForDerivedDataTask& SourceMeshData,
 		const FStaticMeshLODResources& LODModel,
 		class FQueuedThreadPool& ThreadPool,
-		const TArray<FSignedDistanceFieldBuildSectionData>& SectionData,
+		const TArray<FSignedDistanceFieldBuildMaterialData>& MaterialBlendModes,
 		const FBoxSphereBounds& Bounds,
 		const FDistanceFieldVolumeData* DistanceFieldVolumeData,
 		int32 MaxLumenMeshCards,
diff --git a/Engine/Source/Developer/MeshUtilities/Public/MeshUtilities.h b/Engine/Source/Developer/MeshUtilities/Public/MeshUtilities.h
--- a/Engine/Source/Developer/MeshUtilities/Public/MeshUtilities.h
+++ b/Engine/Source/Developer/MeshUtilities/Public/MeshUtilities.h
@@ -48,8 +48,9 @@ namespace ETangentOptions
 	};
 };
 
-struct FSignedDistanceFieldBuildSectionData
+class FSignedDistanceFieldBuildMaterialData
 {
+public:
 	EBlendMode BlendMode = BLEND_Opaque;
 	bool bTwoSided = false;
 	bool bAffectDistanceFieldLighting = true;
@@ -129,7 +130,7 @@ public:
 		const FSourceMeshDataForDerivedDataTask& SourceMeshData,
 		const FStaticMeshLODResources& LODModel,
 		class FQueuedThreadPool& ThreadPool,
-		const TArray<FSignedDistanceFieldBuildSectionData>& SectionData,
+		const TArray<FSignedDistanceFieldBuildMaterialData>& MaterialBlendModes,
 		const FBoxSphereBounds& Bounds,
 		float DistanceFieldResolutionScale,
 		bool bGenerateAsIfTwoSided,
@@ -140,7 +141,7 @@ public:
 		const FSourceMeshDataForDerivedDataTask& SourceMeshData,
 		const FStaticMeshLODResources& LODModel,
 		class FQueuedThreadPool& ThreadPool,
-		const TArray<FSignedDistanceFieldBuildSectionData>& SectionData,
+		const TArray<FSignedDistanceFieldBuildMaterialData>& MaterialBlendModes,
 		const FBoxSphereBounds& Bounds,
 		const class FDistanceFieldVolumeData* DistanceFieldVolumeData,
 		int32 MaxLumenMeshCards,
diff --git a/Engine/Source/Runtime/Engine/Private/DistanceFieldAtlas.cpp b/Engine/Source/Runtime/Engine/Private/DistanceFieldAtlas.cpp
--- a/Engine/Source/Runtime/Engine/Private/DistanceFieldAtlas.cpp
+++ b/Engine/Source/Runtime/Engine/Private/DistanceFieldAtlas.cpp
@@ -125,7 +125,7 @@ FDistanceFieldAsyncQueue* GDistanceFieldAsyncQueue = NULL;
 #if WITH_EDITOR
 
 // DDC key for distance field data, must be changed when modifying the generation code or data format
-#define DISTANCEFIELD_DERIVEDDATA_VER TEXT("E74258BD-3117-4B00-B97C-29DDA01A3D75")
+#define DISTANCEFIELD_DERIVEDDATA_VER TEXT("295895D9-FFCE-48B5-9A4B-131DB318441E")
 
 FString BuildDistanceFieldDerivedDataKey(const FString& InMeshKey)
 {
@@ -147,52 +147,46 @@ FString BuildDistanceFieldDerivedDataKey(const FString& InMeshKey)
 
 #if WITH_EDITORONLY_DATA
 
-void BuildSignedDistanceFieldBuildSectionData(UStaticMesh* Mesh, uint32 LODIndex, TArray<FSignedDistanceFieldBuildSectionData>& OutData)
+void FDistanceFieldVolumeData::CacheDerivedData(const FString& InStaticMeshDerivedDataKey, const ITargetPlatform* TargetPlatform, UStaticMesh* Mesh, FStaticMeshRenderData& RenderData, UStaticMesh* GenerateSource, float DistanceFieldResolutionScale, bool bGenerateDistanceFieldAsIfTwoSided)
 {
+	TRACE_CPUPROFILER_EVENT_SCOPE(FDistanceFieldVolumeData::CacheDerivedData);
+
 	const TArray<FStaticMaterial>& StaticMaterials = Mesh->GetStaticMaterials();
-	const FMeshSectionInfoMap& SectionInfoMap = Mesh->GetSectionInfoMap();
 
-	OutData.SetNum(SectionInfoMap.GetSectionNumber(LODIndex));
+	TArray<FSignedDistanceFieldBuildMaterialData> BuildMaterialData;
+	BuildMaterialData.SetNum(StaticMaterials.Num());
 
-	for (int32 SectionIndex = 0; SectionIndex < OutData.Num(); SectionIndex++)
+	const FMeshSectionInfoMap& SectionInfoMap = Mesh->GetSectionInfoMap();
+	const uint32 LODIndex = 0;
+
+	for (int32 SectionIndex = 0; SectionIndex < SectionInfoMap.GetSectionNumber(LODIndex); SectionIndex++)
 	{
 		const FMeshSectionInfo& Section = SectionInfoMap.Get(LODIndex, SectionIndex);
 
-		FSignedDistanceFieldBuildSectionData& SectionData = OutData[SectionIndex];
-		SectionData.bAffectDistanceFieldLighting = Section.bAffectDistanceFieldLighting;
-
-		if (!StaticMaterials.IsValidIndex(Section.MaterialIndex))
+		if (!BuildMaterialData.IsValidIndex(Section.MaterialIndex))
 		{
-			// TODO: Should maybe log warning here
 			continue;
 		}
 
+		FSignedDistanceFieldBuildMaterialData& MaterialData = BuildMaterialData[Section.MaterialIndex];
+		MaterialData.bAffectDistanceFieldLighting = Section.bAffectDistanceFieldLighting;
+
 		UMaterialInterface* MaterialInterface = StaticMaterials[Section.MaterialIndex].MaterialInterface;
 		if (MaterialInterface)
 		{
-			SectionData.BlendMode = MaterialInterface->GetBlendMode();
-			SectionData.bTwoSided = MaterialInterface->IsTwoSided();
+			MaterialData.BlendMode = MaterialInterface->GetBlendMode();
+			MaterialData.bTwoSided = MaterialInterface->IsTwoSided();
 		}
 	}
-}
-
-void FDistanceFieldVolumeData::CacheDerivedData(const FString& InStaticMeshDerivedDataKey, const ITargetPlatform* TargetPlatform, UStaticMesh* Mesh, FStaticMeshRenderData& RenderData, UStaticMesh* GenerateSource, float DistanceFieldResolutionScale, bool bGenerateDistanceFieldAsIfTwoSided)
-{
-	TRACE_CPUPROFILER_EVENT_SCOPE(FDistanceFieldVolumeData::CacheDerivedData);
-
-	TArray<FSignedDistanceFieldBuildSectionData> BuildSectionData;
-
-	const uint32 LODIndex = 0;
-	BuildSignedDistanceFieldBuildSectionData(GenerateSource, LODIndex, BuildSectionData);
 
 	FString DistanceFieldKey = BuildDistanceFieldDerivedDataKey(InStaticMeshDerivedDataKey);
 
-	for (int32 SectionIndex = 0; SectionIndex < BuildSectionData.Num(); SectionIndex++)
+	for (int32 MaterialIndex = 0; MaterialIndex < Mesh->GetStaticMaterials().Num(); MaterialIndex++)
 	{
 		DistanceFieldKey += FString::Printf(TEXT("_M%u_%u_%u"), 
-			(uint32)BuildSectionData[SectionIndex].BlendMode,
-			BuildSectionData[SectionIndex].bTwoSided ? 1 : 0,
-			BuildSectionData[SectionIndex].bAffectDistanceFieldLighting ? 1 : 0);
+			(uint32)BuildMaterialData[MaterialIndex].BlendMode,
+			BuildMaterialData[MaterialIndex].bTwoSided ? 1 : 0,
+			BuildMaterialData[MaterialIndex].bAffectDistanceFieldLighting ? 1 : 0);
 	}
 
 	TArray<uint8> DerivedData;
@@ -222,13 +216,13 @@ void FDistanceFieldVolumeData::CacheDerivedData(const FString& InStaticMeshDeriv
 		NewTask->GeneratedVolumeData = new FDistanceFieldVolumeData();
 		NewTask->GeneratedVolumeData->AssetName = Mesh->GetFName();
 		NewTask->GeneratedVolumeData->bAsyncBuilding = true;
-		NewTask->SectionData = MoveTemp(BuildSectionData);
+		NewTask->MaterialBlendModes = MoveTemp(BuildMaterialData);
 
 		// Nanite overrides source static mesh with a coarse representation. Need to load original data before we build the mesh SDF.
-		if (GenerateSource->IsNaniteEnabled())
+		if (Mesh->IsNaniteEnabled())
 		{
 			IMeshBuilderModule& MeshBuilderModule = IMeshBuilderModule::GetForPlatform(TargetPlatform);
-			if (!MeshBuilderModule.BuildMeshVertexPositions(GenerateSource, NewTask->SourceMeshData.TriangleIndices, NewTask->SourceMeshData.VertexPositions, NewTask->SourceMeshData.Sections))
+			if (!MeshBuilderModule.BuildMeshVertexPositions(Mesh, NewTask->SourceMeshData.TriangleIndices, NewTask->SourceMeshData.VertexPositions))
 			{
 				UE_LOG(LogStaticMesh, Error, TEXT("Failed to build static mesh. See previous line(s) for details."));
 			}
@@ -724,7 +718,7 @@ void FDistanceFieldAsyncQueue::Build(FAsyncDistanceFieldTask* Task, FQueuedThrea
 			Task->SourceMeshData,
 			LODModel,
 			BuildThreadPool,
-			Task->SectionData,
+			Task->MaterialBlendModes,
 			Task->GenerateSource->GetRenderData()->Bounds,
 			Task->DistanceFieldResolutionScale,
 			Task->bGenerateDistanceFieldAsIfTwoSided,
diff --git a/Engine/Source/Runtime/Engine/Private/MeshCardRepresentation.cpp b/Engine/Source/Runtime/Engine/Private/MeshCardRepresentation.cpp
--- a/Engine/Source/Runtime/Engine/Private/MeshCardRepresentation.cpp
+++ b/Engine/Source/Runtime/Engine/Private/MeshCardRepresentation.cpp
@@ -98,7 +98,7 @@ FCardRepresentationAsyncQueue* GCardRepresentationAsyncQueue = NULL;
 #if WITH_EDITOR
 
 // DDC key for card representation data, must be changed when modifying the generation code or data format
-#define CARDREPRESENTATION_DERIVEDDATA_VER TEXT("5A576281-130D-4159-B20C-9A9880228C88")
+#define CARDREPRESENTATION_DERIVEDDATA_VER TEXT("FF1E9B99-1837-4F13-A892-13BD62922D0B")
 
 FString BuildCardRepresentationDerivedDataKey(const FString& InMeshKey, int32 MaxLumenMeshCards)
 {
@@ -117,8 +117,6 @@ FString BuildCardRepresentationDerivedDataKey(const FString& InMeshKey, int32 Ma
 
 #if WITH_EDITORONLY_DATA
 
-extern void BuildSignedDistanceFieldBuildSectionData(UStaticMesh* Mesh, uint32 LODIndex, TArray<FSignedDistanceFieldBuildSectionData>& OutData);
-
 void BeginCacheMeshCardRepresentation(const ITargetPlatform* TargetPlatform, UStaticMesh* StaticMeshAsset, FStaticMeshRenderData& RenderData, const FString& DistanceFieldKey, FSourceMeshDataForDerivedDataTask* OptionalSourceMeshData)
 {
 	static const auto CVarCards = IConsoleManager::Get().FindTConsoleVariableDataInt(TEXT("r.MeshCardRepresentation"));
@@ -175,9 +173,31 @@ void FCardRepresentationData::CacheDerivedData(const FString& InDDCKey, const IT
 		NewTask->GeneratedCardRepresentation = new FCardRepresentationData();
 		NewTask->MaxLumenMeshCards = MaxLumenMeshCards;
 		NewTask->bGenerateDistanceFieldAsIfTwoSided = bGenerateDistanceFieldAsIfTwoSided;
+		NewTask->MaterialBlendModes.SetNum(Mesh->GetStaticMaterials().Num());
 
+		const TArray<FStaticMaterial>& StaticMaterials = Mesh->GetStaticMaterials();
+		const FMeshSectionInfoMap& SectionInfoMap = Mesh->GetSectionInfoMap();
 		const uint32 LODIndex = 0;
-		BuildSignedDistanceFieldBuildSectionData(Mesh, LODIndex, NewTask->SectionData);
+
+		for (int32 SectionIndex = 0; SectionIndex < SectionInfoMap.GetSectionNumber(LODIndex); SectionIndex++)
+		{
+			const FMeshSectionInfo& Section = SectionInfoMap.Get(LODIndex, SectionIndex);
+
+			if (!NewTask->MaterialBlendModes.IsValidIndex(Section.MaterialIndex))
+			{
+				continue;
+			}
+
+			FSignedDistanceFieldBuildMaterialData& MaterialData = NewTask->MaterialBlendModes[Section.MaterialIndex];
+			MaterialData.bAffectDistanceFieldLighting = Section.bAffectDistanceFieldLighting;
+
+			UMaterialInterface* MaterialInterface = StaticMaterials[Section.MaterialIndex].MaterialInterface;
+			if (MaterialInterface)
+			{
+				MaterialData.BlendMode = MaterialInterface->GetBlendMode();
+				MaterialData.bTwoSided = MaterialInterface->IsTwoSided();
+			}
+		}
 
 		// Nanite overrides source static mesh with a coarse representation. Need to load original data before we build the mesh SDF.
 		if (OptionalSourceMeshData)
@@ -187,7 +207,7 @@ void FCardRepresentationData::CacheDerivedData(const FString& InDDCKey, const IT
 		else if (Mesh->IsNaniteEnabled())
 		{
 			IMeshBuilderModule& MeshBuilderModule = IMeshBuilderModule::GetForPlatform(TargetPlatform);
-			if (!MeshBuilderModule.BuildMeshVertexPositions(Mesh, NewTask->SourceMeshData.TriangleIndices, NewTask->SourceMeshData.VertexPositions, NewTask->SourceMeshData.Sections))
+			if (!MeshBuilderModule.BuildMeshVertexPositions(Mesh, NewTask->SourceMeshData.TriangleIndices, NewTask->SourceMeshData.VertexPositions))
 			{
 				UE_LOG(LogStaticMesh, Error, TEXT("Failed to build static mesh. See previous line(s) for details."));
 			}
@@ -652,7 +672,7 @@ void FCardRepresentationAsyncQueue::Build(FAsyncCardRepresentationTask* Task, FQ
 			Task->SourceMeshData,
 			LODModel,
 			BuildThreadPool,
-			Task->SectionData,
+			Task->MaterialBlendModes,
 			Task->GenerateSource->GetRenderData()->Bounds,
 			Task->GenerateSource->GetRenderData()->LODResources[0].DistanceFieldData,
 			Task->MaxLumenMeshCards,
diff --git a/Engine/Source/Runtime/Engine/Public/DerivedMeshDataTaskUtils.h b/Engine/Source/Runtime/Engine/Public/DerivedMeshDataTaskUtils.h
--- a/Engine/Source/Runtime/Engine/Public/DerivedMeshDataTaskUtils.h
+++ b/Engine/Source/Runtime/Engine/Public/DerivedMeshDataTaskUtils.h
@@ -15,7 +15,6 @@ class FSourceMeshDataForDerivedDataTask
 public:
 	TArray<uint32> TriangleIndices;
 	TArray<FVector3f> VertexPositions;
-	FStaticMeshSectionArray Sections;
 
 	int32 GetNumIndices() const
 	{
diff --git a/Engine/Source/Runtime/Engine/Public/DistanceFieldAtlas.h b/Engine/Source/Runtime/Engine/Public/DistanceFieldAtlas.h
--- a/Engine/Source/Runtime/Engine/Public/DistanceFieldAtlas.h
+++ b/Engine/Source/Runtime/Engine/Public/DistanceFieldAtlas.h
@@ -334,7 +334,7 @@ public:
 	FAsyncDistanceFieldTask();
 
 #if WITH_EDITOR
-	TArray<FSignedDistanceFieldBuildSectionData> SectionData;
+	TArray<FSignedDistanceFieldBuildMaterialData> MaterialBlendModes;
 #endif
 	FSourceMeshDataForDerivedDataTask SourceMeshData;
 	UStaticMesh* StaticMesh;
diff --git a/Engine/Source/Runtime/Engine/Public/MeshCardBuild.h b/Engine/Source/Runtime/Engine/Public/MeshCardBuild.h
--- a/Engine/Source/Runtime/Engine/Public/MeshCardBuild.h
+++ b/Engine/Source/Runtime/Engine/Public/MeshCardBuild.h
@@ -12,7 +12,7 @@
 #include <atomic>
 
 class UStaticMesh;
-struct FSignedDistanceFieldBuildSectionData;
+class FSignedDistanceFieldBuildMaterialData;
 
 class FLumenCardBuildData
 {
@@ -190,7 +190,7 @@ public:
 	bool bSuccess = false;
 
 #if WITH_EDITOR
-	TArray<FSignedDistanceFieldBuildSectionData> SectionData;
+	TArray<FSignedDistanceFieldBuildMaterialData> MaterialBlendModes;
 #endif
 
 	FSourceMeshDataForDerivedDataTask SourceMeshData;
