Resubmit CL26160360, 26146080, 26137852, 26126752

Only leaving CL 26157603 reverted since accessing GenerateSource in FDistanceFieldVolumeData::CacheDerivedData is not safe due to async mesh compilation.

#prelight 64948190a26e02126bc0a64f

diff --git a/Engine/Source/Developer/MeshBuilder/Private/MeshBuilderModule.cpp b/Engine/Source/Developer/MeshBuilder/Private/MeshBuilderModule.cpp
--- a/Engine/Source/Developer/MeshBuilder/Private/MeshBuilderModule.cpp
+++ b/Engine/Source/Developer/MeshBuilder/Private/MeshBuilderModule.cpp
@@ -31,7 +31,8 @@ public:
 	virtual bool BuildMeshVertexPositions(
 		UObject* StaticMesh,
 		TArray<uint32>& Indices,
-		TArray<FVector3f>& Vertices) override;
+		TArray<FVector3f>& Vertices,
+		FStaticMeshSectionArray& Sections) override;
 
 	virtual bool BuildSkeletalMesh(const FSkeletalMeshBuildParameters& SkeletalMeshBuildParameters) override;
 
@@ -55,13 +56,14 @@ bool FMeshBuilderModule::BuildMesh(FStaticMeshRenderData& OutRenderData, class U
 bool FMeshBuilderModule::BuildMeshVertexPositions(
 	UObject* Mesh,
 	TArray<uint32>& Indices,
-	TArray<FVector3f>& Vertices)
+	TArray<FVector3f>& Vertices,
+	FStaticMeshSectionArray& Sections)
 {
 	UStaticMesh* StaticMesh = Cast< UStaticMesh >(Mesh);
 	if (StaticMesh)
 	{
 		//Call the static mesh builder
-		return FStaticMeshBuilder().BuildMeshVertexPositions(StaticMesh, Indices, Vertices);
+		return FStaticMeshBuilder().BuildMeshVertexPositions(StaticMesh, Indices, Vertices, Sections);
 	}
 	return false;
 }
diff --git a/Engine/Source/Developer/MeshBuilder/Private/StaticMeshBuilder.cpp b/Engine/Source/Developer/MeshBuilder/Private/StaticMeshBuilder.cpp
--- a/Engine/Source/Developer/MeshBuilder/Private/StaticMeshBuilder.cpp
+++ b/Engine/Source/Developer/MeshBuilder/Private/StaticMeshBuilder.cpp
@@ -730,7 +730,8 @@ bool FStaticMeshBuilder::Build(FStaticMeshRenderData& StaticMeshRenderData, USta
 bool FStaticMeshBuilder::BuildMeshVertexPositions(
 	UStaticMesh* StaticMesh,
 	TArray<uint32>& BuiltIndices,
-	TArray<FVector3f>& BuiltVertices)
+	TArray<FVector3f>& BuiltVertices,
+	FStaticMeshSectionArray& Sections)
 {
 	TRACE_CPUPROFILER_EVENT_SCOPE(FStaticMeshBuilder::BuildMeshVertexPositions);
 
@@ -742,59 +743,62 @@ bool FStaticMeshBuilder::BuildMeshVertexPositions(
 		return false;
 	}
 
-	const int32 NumSourceModels = StaticMesh->GetNumSourceModels();
-	if (NumSourceModels > 0)
-	{
-		FMeshDescription MeshDescription;
-		const bool bIsMeshDescriptionValid = SourceModel.CloneMeshDescription(MeshDescription);
-		if (bIsMeshDescriptionValid)
-		{
-			FElementIDRemappings Remappings;
-			if (MeshDescription.NeedsCompact())
-			{
-				MeshDescription.Compact(Remappings);
-			}
+	FMeshDescription MeshDescription;
+	const bool bIsMeshDescriptionValid = SourceModel.CloneMeshDescription(MeshDescription);
+	check(bIsMeshDescriptionValid);
 
-			const FMeshBuildSettings& BuildSettings = SourceModel.BuildSettings;
+	FMeshBuildSettings& BuildSettings = StaticMesh->GetSourceModel(0).BuildSettings;
 
-			const FStaticMeshConstAttributes Attributes(MeshDescription);
-			TArrayView<const FVector3f> VertexPositions = Attributes.GetVertexPositions().GetRawArray();
-			TArrayView<const FVertexID> VertexIndices = Attributes.GetTriangleVertexIndices().GetRawArray();
-			const FVector3f BuildScale3D = (FVector3f)BuildSettings.BuildScale3D;
+	FMeshDescriptionHelper MeshDescriptionHelper(&BuildSettings);
+	MeshDescriptionHelper.SetupRenderMeshDescription(StaticMesh, MeshDescription, false, false);
 
-			BuiltVertices.Reserve(VertexPositions.Num());
-			for (int32 VertexIndex = 0; VertexIndex < VertexPositions.Num(); ++VertexIndex)
-			{
-				BuiltVertices.Add(VertexPositions[VertexIndex] * BuildScale3D);
-			}
+	const FPolygonGroupArray& PolygonGroups = MeshDescription.PolygonGroups();
 
-			BuiltIndices.Reserve(VertexIndices.Num());
-			for (int32 TriangleIndex = 0; TriangleIndex < VertexIndices.Num() / 3; ++TriangleIndex)
-			{
-				const uint32 I0 = VertexIndices[TriangleIndex * 3 + 0];
-				const uint32 I1 = VertexIndices[TriangleIndex * 3 + 1];
-				const uint32 I2 = VertexIndices[TriangleIndex * 3 + 2];
+	// Build new vertex buffers
+	FMeshBuildVertexData BuildVertexData;
 
-				if (!ensureMsgf(I0 != INDEX_NONE && I1 != INDEX_NONE && I2 != INDEX_NONE,
-					TEXT("Mesh '%s' has triangles with uninitialized vertex indices"), *StaticMesh->GetName()))
-				{
-					continue;
-				}
+	Sections.Empty(PolygonGroups.Num());
 
-				const FVector3f V0 = BuiltVertices[I0];
-				const FVector3f V1 = BuiltVertices[I1];
-				const FVector3f V2 = BuiltVertices[I2];
+	TArray<int32> RemapVerts; //Because we will remove MeshVertex that are redundant, we need a remap
+	//Render data Wedge map is only set for LOD 0???
 
-				const FVector3f TriangleNormal = ((V1 - V2) ^ (V0 - V2));
-				const bool bDegenerateTriangle = TriangleNormal.SizeSquared() < SMALL_NUMBER;
-				if (!bDegenerateTriangle)
-				{
-					BuiltIndices.Add(I0);
-					BuiltIndices.Add(I1);
-					BuiltIndices.Add(I2);
-				}
-			}
-		}
+	TArray<int32> WedgeMap;
+
+	// Prepare the PerSectionIndices array so we can optimize the index buffer for the GPU
+	TArray<TArray<uint32>> PerSectionIndices;
+	PerSectionIndices.AddDefaulted(MeshDescription.PolygonGroups().Num());
+
+	FBoxSphereBounds LODBounds;
+
+	// Build the vertex and index buffer
+	UE::Private::StaticMeshBuilder::BuildVertexBuffer(
+		StaticMesh,
+		MeshDescription,
+		BuildSettings,
+		WedgeMap,
+		Sections,
+		PerSectionIndices,
+		BuildVertexData,
+		MeshDescriptionHelper.GetOverlappingCorners(),
+		RemapVerts,
+		LODBounds,
+		false /* bNeedTangents */,
+		false /* bNeedWedgeMap */
+	);
+
+	BuiltVertices = BuildVertexData.Position;
+
+	// Release MeshDescription memory since we don't need it anymore
+	MeshDescription.Empty();
+
+	// Concatenate the per-section index buffers.
+	bool bNeeds32BitIndices = false;
+	UE::Private::StaticMeshBuilder::BuildCombinedSectionIndices(PerSectionIndices, Sections, BuiltIndices, bNeeds32BitIndices);
+
+	// Apply section remapping
+	for (int32 SectionIndex = 0; SectionIndex < Sections.Num(); SectionIndex++)
+	{
+		Sections[SectionIndex].MaterialIndex = StaticMesh->GetSectionInfoMap().Get(0, SectionIndex).MaterialIndex;
 	}
 
 	return true;
diff --git a/Engine/Source/Developer/MeshBuilder/Public/IMeshBuilderModule.h b/Engine/Source/Developer/MeshBuilder/Public/IMeshBuilderModule.h
--- a/Engine/Source/Developer/MeshBuilder/Public/IMeshBuilderModule.h
+++ b/Engine/Source/Developer/MeshBuilder/Public/IMeshBuilderModule.h
@@ -31,7 +31,7 @@ public:
 
 	virtual bool BuildMesh(class FStaticMeshRenderData& OutRenderData, class UObject* Mesh, const class FStaticMeshLODGroup& LODGroup, bool bGenerateCoarseMeshStreamingLODs, bool bAllowNanite) = 0;
 
-	virtual bool BuildMeshVertexPositions(class UObject* StaticMesh, TArray<uint32>& Indices, TArray<FVector3f>& Vertices) = 0;
+	virtual bool BuildMeshVertexPositions(class UObject* StaticMesh, TArray<uint32>& Indices, TArray<FVector3f>& Vertices, FStaticMeshSectionArray& Sections) = 0;
 
 	virtual bool BuildSkeletalMesh(const struct FSkeletalMeshBuildParameters& SkeletalMeshBuildParameters) = 0;
 
diff --git a/Engine/Source/Developer/MeshBuilder/Public/MeshBuilder.h b/Engine/Source/Developer/MeshBuilder/Public/MeshBuilder.h
--- a/Engine/Source/Developer/MeshBuilder/Public/MeshBuilder.h
+++ b/Engine/Source/Developer/MeshBuilder/Public/MeshBuilder.h
@@ -38,7 +38,8 @@ public:
 	virtual bool BuildMeshVertexPositions(
 		UStaticMesh* StaticMesh,
 		TArray<uint32>& Indices,
-		TArray<FVector3f>& Vertices) = 0;
+		TArray<FVector3f>& Vertices,
+		FStaticMeshSectionArray& Sections) = 0;
 
 	/**
 	 * Build function should be override and is the starting point for skeletal mesh builders
diff --git a/Engine/Source/Developer/MeshBuilder/Public/SkeletalMeshBuilder.h b/Engine/Source/Developer/MeshBuilder/Public/SkeletalMeshBuilder.h
--- a/Engine/Source/Developer/MeshBuilder/Public/SkeletalMeshBuilder.h
+++ b/Engine/Source/Developer/MeshBuilder/Public/SkeletalMeshBuilder.h
@@ -28,7 +28,8 @@ public:
 	virtual bool BuildMeshVertexPositions(
 		UStaticMesh* StaticMesh,
 		TArray<uint32>& Indices,
-		TArray<FVector3f>& Vertices) override
+		TArray<FVector3f>& Vertices,
+		FStaticMeshSectionArray& Sections) override
 	{
 		bool No_Support_For_StaticMesh_Build_In_FSkeletalMeshBuilder_Class = false;
 		check(No_Support_For_StaticMesh_Build_In_FSkeletalMeshBuilder_Class);
diff --git a/Engine/Source/Developer/MeshBuilder/Public/StaticMeshBuilder.h b/Engine/Source/Developer/MeshBuilder/Public/StaticMeshBuilder.h
--- a/Engine/Source/Developer/MeshBuilder/Public/StaticMeshBuilder.h
+++ b/Engine/Source/Developer/MeshBuilder/Public/StaticMeshBuilder.h
@@ -40,7 +40,8 @@ public:
 	virtual bool BuildMeshVertexPositions(
 		UStaticMesh* StaticMesh,
 		TArray<uint32>& Indices,
-		TArray<FVector3f>& Vertices) override;
+		TArray<FVector3f>& Vertices,
+		FStaticMeshSectionArray& Sections) override;
 
 private:
 
diff --git a/Engine/Source/Developer/MeshUtilities/Private/MeshCardRepresentationUtilities.cpp b/Engine/Source/Developer/MeshUtilities/Private/MeshCardRepresentationUtilities.cpp
--- a/Engine/Source/Developer/MeshUtilities/Private/MeshCardRepresentationUtilities.cpp
+++ b/Engine/Source/Developer/MeshUtilities/Private/MeshCardRepresentationUtilities.cpp
@@ -1054,7 +1054,7 @@ bool FMeshUtilities::GenerateCardRepresentationData(
 	const FSourceMeshDataForDerivedDataTask& SourceMeshData,
 	const FStaticMeshLODResources& LODModel,
 	class FQueuedThreadPool& ThreadPool,
-	const TArray<FSignedDistanceFieldBuildMaterialData>& MaterialBlendModes,
+	const TArray<FSignedDistanceFieldBuildSectionData>& SectionData,
 	const FBoxSphereBounds& Bounds,
 	const FDistanceFieldVolumeData* DistanceFieldVolumeData,
 	int32 MaxLumenMeshCards,
@@ -1072,7 +1072,7 @@ bool FMeshUtilities::GenerateCardRepresentationData(
 		MeshRepresentation::SetupEmbreeScene(MeshName,
 			SourceMeshData,
 			LODModel,
-			MaterialBlendModes,
+			SectionData,
 			bGenerateAsIfTwoSided,
 			EmbreeScene);
 
diff --git a/Engine/Source/Developer/MeshUtilities/Private/MeshDistanceFieldUtilities.cpp b/Engine/Source/Developer/MeshUtilities/Private/MeshDistanceFieldUtilities.cpp
--- a/Engine/Source/Developer/MeshUtilities/Private/MeshDistanceFieldUtilities.cpp
+++ b/Engine/Source/Developer/MeshUtilities/Private/MeshDistanceFieldUtilities.cpp
@@ -252,7 +252,7 @@ void FMeshUtilities::GenerateSignedDistanceFieldVolumeData(
 	const FSourceMeshDataForDerivedDataTask& SourceMeshData,
 	const FStaticMeshLODResources& LODModel,
 	class FQueuedThreadPool& ThreadPool,
-	const TArray<FSignedDistanceFieldBuildMaterialData>& MaterialBlendModes,
+	const TArray<FSignedDistanceFieldBuildSectionData>& SectionData,
 	const FBoxSphereBounds& Bounds,
 	float DistanceFieldResolutionScale,
 	bool bGenerateAsIfTwoSided,
@@ -268,7 +268,7 @@ void FMeshUtilities::GenerateSignedDistanceFieldVolumeData(
 		MeshRepresentation::SetupEmbreeScene(MeshName,
 			SourceMeshData,
 			LODModel,
-			MaterialBlendModes,
+			SectionData,
 			bGenerateAsIfTwoSided,
 			EmbreeScene);
 
@@ -529,7 +529,7 @@ void FMeshUtilities::GenerateSignedDistanceFieldVolumeData(
 	const FSourceMeshDataForDerivedDataTask& SourceMeshData,
 	const FStaticMeshLODResources& LODModel,
 	class FQueuedThreadPool& ThreadPool,
-	const TArray<FSignedDistanceFieldBuildMaterialData>& MaterialBlendModes,
+	const TArray<FSignedDistanceFieldBuildSectionData>& SectionData,
 	const FBoxSphereBounds& Bounds,
 	float DistanceFieldResolutionScale,
 	bool bGenerateAsIfTwoSided,
diff --git a/Engine/Source/Developer/MeshUtilities/Private/MeshRepresentationCommon.cpp b/Engine/Source/Developer/MeshUtilities/Private/MeshRepresentationCommon.cpp
--- a/Engine/Source/Developer/MeshUtilities/Private/MeshRepresentationCommon.cpp
+++ b/Engine/Source/Developer/MeshUtilities/Private/MeshRepresentationCommon.cpp
@@ -126,15 +126,17 @@ void MeshRepresentation::SetupEmbreeScene(
 	FString MeshName,
 	const FSourceMeshDataForDerivedDataTask& SourceMeshData,
 	const FStaticMeshLODResources& LODModel,
-	const TArray<FSignedDistanceFieldBuildMaterialData>& MaterialBlendModes,
+	const TArray<FSignedDistanceFieldBuildSectionData>& SectionData,
 	bool bGenerateAsIfTwoSided,
 	FEmbreeScene& EmbreeScene)
 {
+	const uint32 NumVertices = SourceMeshData.IsValid() ? SourceMeshData.GetNumVertices() : LODModel.VertexBuffers.PositionVertexBuffer.GetNumVertices();
 	const uint32 NumIndices = SourceMeshData.IsValid() ? SourceMeshData.GetNumIndices() : LODModel.IndexBuffer.GetNumIndices();
 	const int32 NumTriangles = NumIndices / 3;
-	const uint32 NumVertices = SourceMeshData.IsValid() ? SourceMeshData.GetNumVertices() : LODModel.VertexBuffers.PositionVertexBuffer.GetNumVertices();
 	EmbreeScene.NumIndices = NumTriangles;
 
+	const FStaticMeshSectionArray& Sections = SourceMeshData.IsValid() ? SourceMeshData.Sections : LODModel.Sections;
+
 	TArray<FkDOPBuildCollisionTriangle<uint32> > BuildTriangles;
 
 #if USE_EMBREE
@@ -165,68 +167,69 @@ void MeshRepresentation::SetupEmbreeScene(
 	}
 #endif
 
+	/*
+	if (LODModel.Sections.Num() > SectionData.Num())
+	{
+		UE_LOG(LogMeshUtilities, Warning, TEXT("Unexpected number of mesh sections when setting up Embree Scene for %s."), *MeshName);
+	}
+	*/
+
 	TArray<int32> FilteredTriangles;
 	FilteredTriangles.Empty(NumTriangles);
 
-	if (SourceMeshData.IsValid())
+	for (int32 TriangleIndex = 0; TriangleIndex < NumTriangles; ++TriangleIndex)
 	{
-		for (int32 TriangleIndex = 0; TriangleIndex < NumTriangles; ++TriangleIndex)
+		FVector3f V0, V1, V2;
+
+		if (SourceMeshData.IsValid())
 		{
 			const uint32 I0 = SourceMeshData.TriangleIndices[TriangleIndex * 3 + 0];
 			const uint32 I1 = SourceMeshData.TriangleIndices[TriangleIndex * 3 + 1];
 			const uint32 I2 = SourceMeshData.TriangleIndices[TriangleIndex * 3 + 2];
 
-			const FVector3f V0 = SourceMeshData.VertexPositions[I0];
-			const FVector3f V1 = SourceMeshData.VertexPositions[I1];
-			const FVector3f V2 = SourceMeshData.VertexPositions[I2];
-
-			const FVector3f TriangleNormal = ((V1 - V2) ^ (V0 - V2));
-			const bool bDegenerateTriangle = TriangleNormal.SizeSquared() < SMALL_NUMBER;
-			if (!bDegenerateTriangle)
-			{
-				FilteredTriangles.Add(TriangleIndex);
-			}
+			V0 = SourceMeshData.VertexPositions[I0];
+			V1 = SourceMeshData.VertexPositions[I1];
+			V2 = SourceMeshData.VertexPositions[I2];
 		}
-	}
-	else
-	{
-		for (int32 TriangleIndex = 0; TriangleIndex < NumTriangles; ++TriangleIndex)
+		else
 		{
 			const FIndexArrayView Indices = LODModel.IndexBuffer.GetArrayView();
 			const uint32 I0 = Indices[TriangleIndex * 3 + 0];
 			const uint32 I1 = Indices[TriangleIndex * 3 + 1];
 			const uint32 I2 = Indices[TriangleIndex * 3 + 2];
 
-			const FVector3f V0 = LODModel.VertexBuffers.PositionVertexBuffer.VertexPosition(I0);
-			const FVector3f V1 = LODModel.VertexBuffers.PositionVertexBuffer.VertexPosition(I1);
-			const FVector3f V2 = LODModel.VertexBuffers.PositionVertexBuffer.VertexPosition(I2);
+			V0 = LODModel.VertexBuffers.PositionVertexBuffer.VertexPosition(I0);
+			V1 = LODModel.VertexBuffers.PositionVertexBuffer.VertexPosition(I1);
+			V2 = LODModel.VertexBuffers.PositionVertexBuffer.VertexPosition(I2);
+		}
+
+		const FVector3f TriangleNormal = ((V1 - V2) ^ (V0 - V2));
+		const bool bDegenerateTriangle = TriangleNormal.SizeSquared() < SMALL_NUMBER;
+		if (!bDegenerateTriangle)
+		{
+			//FilteredTriangles.Add(TriangleIndex);
+
+			bool bTriangleIsOpaqueOrMasked = false;
 
-			const FVector3f TriangleNormal = ((V1 - V2) ^ (V0 - V2));
-			const bool bDegenerateTriangle = TriangleNormal.SizeSquared() < SMALL_NUMBER;
-			if (!bDegenerateTriangle)
+			for (int32 SectionIndex = 0; SectionIndex < Sections.Num(); SectionIndex++)
 			{
-				bool bTriangleIsOpaqueOrMasked = false;
+				const FStaticMeshSection& Section = Sections[SectionIndex];
 
-				for (int32 SectionIndex = 0; SectionIndex < LODModel.Sections.Num(); SectionIndex++)
+				if ((uint32)(TriangleIndex * 3) >= Section.FirstIndex && (uint32)(TriangleIndex * 3) < Section.FirstIndex + Section.NumTriangles * 3)
 				{
-					const FStaticMeshSection& Section = LODModel.Sections[SectionIndex];
-
-					if ((uint32)(TriangleIndex * 3) >= Section.FirstIndex && (uint32)(TriangleIndex * 3) < Section.FirstIndex + Section.NumTriangles * 3)
+					if (SectionData.IsValidIndex(SectionIndex))
 					{
-						if (MaterialBlendModes.IsValidIndex(Section.MaterialIndex))
-						{
-							bTriangleIsOpaqueOrMasked = !IsTranslucentBlendMode(MaterialBlendModes[Section.MaterialIndex].BlendMode) && MaterialBlendModes[Section.MaterialIndex].bAffectDistanceFieldLighting;
-						}
-
-						break;
+						bTriangleIsOpaqueOrMasked = !IsTranslucentBlendMode(SectionData[SectionIndex].BlendMode) && SectionData[SectionIndex].bAffectDistanceFieldLighting;
 					}
-				}
 
-				if (bTriangleIsOpaqueOrMasked)
-				{
-					FilteredTriangles.Add(TriangleIndex);
+					break;
 				}
 			}
+
+			if (bTriangleIsOpaqueOrMasked)
+			{
+				FilteredTriangles.Add(TriangleIndex);
+			}
 		}
 	}
 
@@ -273,15 +276,15 @@ void MeshRepresentation::SetupEmbreeScene(
 
 		bool bTriangleIsTwoSided = false;
 
-		for (int32 SectionIndex = 0; SectionIndex < LODModel.Sections.Num(); SectionIndex++)
+		for (int32 SectionIndex = 0; SectionIndex < Sections.Num(); SectionIndex++)
 		{
-			const FStaticMeshSection& Section = LODModel.Sections[SectionIndex];
+			const FStaticMeshSection& Section = Sections[SectionIndex];
 
 			if ((uint32)(TriangleIndex * 3) >= Section.FirstIndex && (uint32)(TriangleIndex * 3) < Section.FirstIndex + Section.NumTriangles * 3)
 			{
-				if (MaterialBlendModes.IsValidIndex(Section.MaterialIndex))
+				if (SectionData.IsValidIndex(SectionIndex))
 				{
-					bTriangleIsTwoSided = MaterialBlendModes[Section.MaterialIndex].bTwoSided;
+					bTriangleIsTwoSided = SectionData[SectionIndex].bTwoSided;
 				}
 
 				break;
@@ -350,15 +353,15 @@ void MeshRepresentation::SetupEmbreeScene(
 		uint32 NumTrianglesTotal = 0;
 		uint32 NumTwoSidedTriangles = 0;
 
-		for (int32 SectionIndex = 0; SectionIndex < LODModel.Sections.Num(); SectionIndex++)
+		for (int32 SectionIndex = 0; SectionIndex < Sections.Num(); SectionIndex++)
 		{
-			const FStaticMeshSection& Section = LODModel.Sections[SectionIndex];
+			const FStaticMeshSection& Section = Sections[SectionIndex];
 
-			if (MaterialBlendModes.IsValidIndex(Section.MaterialIndex))
+			if (SectionData.IsValidIndex(SectionIndex))
 			{
 				NumTrianglesTotal += Section.NumTriangles;
 
-				if (MaterialBlendModes[Section.MaterialIndex].bTwoSided)
+				if (SectionData[SectionIndex].bTwoSided)
 				{
 					NumTwoSidedTriangles += Section.NumTriangles;
 				}
diff --git a/Engine/Source/Developer/MeshUtilities/Private/MeshRepresentationCommon.h b/Engine/Source/Developer/MeshUtilities/Private/MeshRepresentationCommon.h
--- a/Engine/Source/Developer/MeshUtilities/Private/MeshRepresentationCommon.h
+++ b/Engine/Source/Developer/MeshUtilities/Private/MeshRepresentationCommon.h
@@ -159,7 +159,7 @@ namespace MeshRepresentation
 	void SetupEmbreeScene(FString MeshName,
 		const FSourceMeshDataForDerivedDataTask& SourceMeshData,
 		const FStaticMeshLODResources& LODModel,
-		const TArray<FSignedDistanceFieldBuildMaterialData>& MaterialBlendModes,
+		const TArray<FSignedDistanceFieldBuildSectionData>& SectionData,
 		bool bGenerateAsIfTwoSided,
 		FEmbreeScene& EmbreeScene);
 
diff --git a/Engine/Source/Developer/MeshUtilities/Private/MeshUtilitiesPrivate.h b/Engine/Source/Developer/MeshUtilities/Private/MeshUtilitiesPrivate.h
--- a/Engine/Source/Developer/MeshUtilities/Private/MeshUtilitiesPrivate.h
+++ b/Engine/Source/Developer/MeshUtilities/Private/MeshUtilitiesPrivate.h
@@ -60,7 +60,7 @@ private:
 		const FSourceMeshDataForDerivedDataTask& SourceMeshData,
 		const FStaticMeshLODResources& LODModel,
 		class FQueuedThreadPool& ThreadPool,
-		const TArray<FSignedDistanceFieldBuildMaterialData>& MaterialBlendModes,
+		const TArray<FSignedDistanceFieldBuildSectionData>& SectionData,
 		const FBoxSphereBounds& Bounds,
 		float DistanceFieldResolutionScale,
 		bool bGenerateAsIfTwoSided,
@@ -71,7 +71,7 @@ private:
 		const FSourceMeshDataForDerivedDataTask& SourceMeshData,
 		const FStaticMeshLODResources& LODModel,
 		class FQueuedThreadPool& ThreadPool,
-		const TArray<FSignedDistanceFieldBuildMaterialData>& MaterialBlendModes,
+		const TArray<FSignedDistanceFieldBuildSectionData>& SectionData,
 		const FBoxSphereBounds& Bounds,
 		const FDistanceFieldVolumeData* DistanceFieldVolumeData,
 		int32 MaxLumenMeshCards,
diff --git a/Engine/Source/Developer/MeshUtilities/Public/MeshUtilities.h b/Engine/Source/Developer/MeshUtilities/Public/MeshUtilities.h
--- a/Engine/Source/Developer/MeshUtilities/Public/MeshUtilities.h
+++ b/Engine/Source/Developer/MeshUtilities/Public/MeshUtilities.h
@@ -48,9 +48,8 @@ namespace ETangentOptions
 	};
 };
 
-class FSignedDistanceFieldBuildMaterialData
+struct FSignedDistanceFieldBuildSectionData
 {
-public:
 	EBlendMode BlendMode = BLEND_Opaque;
 	bool bTwoSided = false;
 	bool bAffectDistanceFieldLighting = true;
@@ -130,7 +129,7 @@ public:
 		const FSourceMeshDataForDerivedDataTask& SourceMeshData,
 		const FStaticMeshLODResources& LODModel,
 		class FQueuedThreadPool& ThreadPool,
-		const TArray<FSignedDistanceFieldBuildMaterialData>& MaterialBlendModes,
+		const TArray<FSignedDistanceFieldBuildSectionData>& SectionData,
 		const FBoxSphereBounds& Bounds,
 		float DistanceFieldResolutionScale,
 		bool bGenerateAsIfTwoSided,
@@ -141,7 +140,7 @@ public:
 		const FSourceMeshDataForDerivedDataTask& SourceMeshData,
 		const FStaticMeshLODResources& LODModel,
 		class FQueuedThreadPool& ThreadPool,
-		const TArray<FSignedDistanceFieldBuildMaterialData>& MaterialBlendModes,
+		const TArray<FSignedDistanceFieldBuildSectionData>& SectionData,
 		const FBoxSphereBounds& Bounds,
 		const class FDistanceFieldVolumeData* DistanceFieldVolumeData,
 		int32 MaxLumenMeshCards,
diff --git a/Engine/Source/Runtime/Engine/Private/DistanceFieldAtlas.cpp b/Engine/Source/Runtime/Engine/Private/DistanceFieldAtlas.cpp
--- a/Engine/Source/Runtime/Engine/Private/DistanceFieldAtlas.cpp
+++ b/Engine/Source/Runtime/Engine/Private/DistanceFieldAtlas.cpp
@@ -125,7 +125,7 @@ FDistanceFieldAsyncQueue* GDistanceFieldAsyncQueue = NULL;
 #if WITH_EDITOR
 
 // DDC key for distance field data, must be changed when modifying the generation code or data format
-#define DISTANCEFIELD_DERIVEDDATA_VER TEXT("295895D9-FFCE-48B5-9A4B-131DB318441E")
+#define DISTANCEFIELD_DERIVEDDATA_VER TEXT("235D7D36-C204-44B7-A44C-D9C67359075C")
 
 FString BuildDistanceFieldDerivedDataKey(const FString& InMeshKey)
 {
@@ -147,46 +147,52 @@ FString BuildDistanceFieldDerivedDataKey(const FString& InMeshKey)
 
 #if WITH_EDITORONLY_DATA
 
-void FDistanceFieldVolumeData::CacheDerivedData(const FString& InStaticMeshDerivedDataKey, const ITargetPlatform* TargetPlatform, UStaticMesh* Mesh, FStaticMeshRenderData& RenderData, UStaticMesh* GenerateSource, float DistanceFieldResolutionScale, bool bGenerateDistanceFieldAsIfTwoSided)
+void BuildSignedDistanceFieldBuildSectionData(UStaticMesh* Mesh, uint32 LODIndex, TArray<FSignedDistanceFieldBuildSectionData>& OutData)
 {
-	TRACE_CPUPROFILER_EVENT_SCOPE(FDistanceFieldVolumeData::CacheDerivedData);
-
 	const TArray<FStaticMaterial>& StaticMaterials = Mesh->GetStaticMaterials();
-
-	TArray<FSignedDistanceFieldBuildMaterialData> BuildMaterialData;
-	BuildMaterialData.SetNum(StaticMaterials.Num());
-
 	const FMeshSectionInfoMap& SectionInfoMap = Mesh->GetSectionInfoMap();
-	const uint32 LODIndex = 0;
 
-	for (int32 SectionIndex = 0; SectionIndex < SectionInfoMap.GetSectionNumber(LODIndex); SectionIndex++)
+	OutData.SetNum(SectionInfoMap.GetSectionNumber(LODIndex));
+
+	for (int32 SectionIndex = 0; SectionIndex < OutData.Num(); SectionIndex++)
 	{
 		const FMeshSectionInfo& Section = SectionInfoMap.Get(LODIndex, SectionIndex);
 
-		if (!BuildMaterialData.IsValidIndex(Section.MaterialIndex))
+		FSignedDistanceFieldBuildSectionData& SectionData = OutData[SectionIndex];
+		SectionData.bAffectDistanceFieldLighting = Section.bAffectDistanceFieldLighting;
+
+		if (!StaticMaterials.IsValidIndex(Section.MaterialIndex))
 		{
+			// TODO: Should maybe log warning here
 			continue;
 		}
 
-		FSignedDistanceFieldBuildMaterialData& MaterialData = BuildMaterialData[Section.MaterialIndex];
-		MaterialData.bAffectDistanceFieldLighting = Section.bAffectDistanceFieldLighting;
-
 		UMaterialInterface* MaterialInterface = StaticMaterials[Section.MaterialIndex].MaterialInterface;
 		if (MaterialInterface)
 		{
-			MaterialData.BlendMode = MaterialInterface->GetBlendMode();
-			MaterialData.bTwoSided = MaterialInterface->IsTwoSided();
+			SectionData.BlendMode = MaterialInterface->GetBlendMode();
+			SectionData.bTwoSided = MaterialInterface->IsTwoSided();
 		}
 	}
+}
+
+void FDistanceFieldVolumeData::CacheDerivedData(const FString& InStaticMeshDerivedDataKey, const ITargetPlatform* TargetPlatform, UStaticMesh* Mesh, FStaticMeshRenderData& RenderData, UStaticMesh* GenerateSource, float DistanceFieldResolutionScale, bool bGenerateDistanceFieldAsIfTwoSided)
+{
+	TRACE_CPUPROFILER_EVENT_SCOPE(FDistanceFieldVolumeData::CacheDerivedData);
+
+	TArray<FSignedDistanceFieldBuildSectionData> BuildSectionData;
+
+	const uint32 LODIndex = 0;
+	BuildSignedDistanceFieldBuildSectionData(Mesh, LODIndex, BuildSectionData);
 
 	FString DistanceFieldKey = BuildDistanceFieldDerivedDataKey(InStaticMeshDerivedDataKey);
 
-	for (int32 MaterialIndex = 0; MaterialIndex < Mesh->GetStaticMaterials().Num(); MaterialIndex++)
+	for (int32 SectionIndex = 0; SectionIndex < BuildSectionData.Num(); SectionIndex++)
 	{
 		DistanceFieldKey += FString::Printf(TEXT("_M%u_%u_%u"), 
-			(uint32)BuildMaterialData[MaterialIndex].BlendMode,
-			BuildMaterialData[MaterialIndex].bTwoSided ? 1 : 0,
-			BuildMaterialData[MaterialIndex].bAffectDistanceFieldLighting ? 1 : 0);
+			(uint32)BuildSectionData[SectionIndex].BlendMode,
+			BuildSectionData[SectionIndex].bTwoSided ? 1 : 0,
+			BuildSectionData[SectionIndex].bAffectDistanceFieldLighting ? 1 : 0);
 	}
 
 	TArray<uint8> DerivedData;
@@ -216,13 +222,13 @@ void FDistanceFieldVolumeData::CacheDerivedData(const FString& InStaticMeshDeriv
 		NewTask->GeneratedVolumeData = new FDistanceFieldVolumeData();
 		NewTask->GeneratedVolumeData->AssetName = Mesh->GetFName();
 		NewTask->GeneratedVolumeData->bAsyncBuilding = true;
-		NewTask->MaterialBlendModes = MoveTemp(BuildMaterialData);
+		NewTask->SectionData = MoveTemp(BuildSectionData);
 
 		// Nanite overrides source static mesh with a coarse representation. Need to load original data before we build the mesh SDF.
 		if (Mesh->IsNaniteEnabled())
 		{
 			IMeshBuilderModule& MeshBuilderModule = IMeshBuilderModule::GetForPlatform(TargetPlatform);
-			if (!MeshBuilderModule.BuildMeshVertexPositions(Mesh, NewTask->SourceMeshData.TriangleIndices, NewTask->SourceMeshData.VertexPositions))
+			if (!MeshBuilderModule.BuildMeshVertexPositions(Mesh, NewTask->SourceMeshData.TriangleIndices, NewTask->SourceMeshData.VertexPositions, NewTask->SourceMeshData.Sections))
 			{
 				UE_LOG(LogStaticMesh, Error, TEXT("Failed to build static mesh. See previous line(s) for details."));
 			}
@@ -718,7 +724,7 @@ void FDistanceFieldAsyncQueue::Build(FAsyncDistanceFieldTask* Task, FQueuedThrea
 			Task->SourceMeshData,
 			LODModel,
 			BuildThreadPool,
-			Task->MaterialBlendModes,
+			Task->SectionData,
 			Task->GenerateSource->GetRenderData()->Bounds,
 			Task->DistanceFieldResolutionScale,
 			Task->bGenerateDistanceFieldAsIfTwoSided,
diff --git a/Engine/Source/Runtime/Engine/Private/MeshCardRepresentation.cpp b/Engine/Source/Runtime/Engine/Private/MeshCardRepresentation.cpp
--- a/Engine/Source/Runtime/Engine/Private/MeshCardRepresentation.cpp
+++ b/Engine/Source/Runtime/Engine/Private/MeshCardRepresentation.cpp
@@ -98,7 +98,7 @@ FCardRepresentationAsyncQueue* GCardRepresentationAsyncQueue = NULL;
 #if WITH_EDITOR
 
 // DDC key for card representation data, must be changed when modifying the generation code or data format
-#define CARDREPRESENTATION_DERIVEDDATA_VER TEXT("FF1E9B99-1837-4F13-A892-13BD62922D0B")
+#define CARDREPRESENTATION_DERIVEDDATA_VER TEXT("DB55CBB7-BA55-4B0F-8EAA-F2819D33D1FE")
 
 FString BuildCardRepresentationDerivedDataKey(const FString& InMeshKey, int32 MaxLumenMeshCards)
 {
@@ -117,6 +117,8 @@ FString BuildCardRepresentationDerivedDataKey(const FString& InMeshKey, int32 Ma
 
 #if WITH_EDITORONLY_DATA
 
+extern void BuildSignedDistanceFieldBuildSectionData(UStaticMesh* Mesh, uint32 LODIndex, TArray<FSignedDistanceFieldBuildSectionData>& OutData);
+
 void BeginCacheMeshCardRepresentation(const ITargetPlatform* TargetPlatform, UStaticMesh* StaticMeshAsset, FStaticMeshRenderData& RenderData, const FString& DistanceFieldKey, FSourceMeshDataForDerivedDataTask* OptionalSourceMeshData)
 {
 	static const auto CVarCards = IConsoleManager::Get().FindTConsoleVariableDataInt(TEXT("r.MeshCardRepresentation"));
@@ -173,31 +175,9 @@ void FCardRepresentationData::CacheDerivedData(const FString& InDDCKey, const IT
 		NewTask->GeneratedCardRepresentation = new FCardRepresentationData();
 		NewTask->MaxLumenMeshCards = MaxLumenMeshCards;
 		NewTask->bGenerateDistanceFieldAsIfTwoSided = bGenerateDistanceFieldAsIfTwoSided;
-		NewTask->MaterialBlendModes.SetNum(Mesh->GetStaticMaterials().Num());
 
-		const TArray<FStaticMaterial>& StaticMaterials = Mesh->GetStaticMaterials();
-		const FMeshSectionInfoMap& SectionInfoMap = Mesh->GetSectionInfoMap();
 		const uint32 LODIndex = 0;
-
-		for (int32 SectionIndex = 0; SectionIndex < SectionInfoMap.GetSectionNumber(LODIndex); SectionIndex++)
-		{
-			const FMeshSectionInfo& Section = SectionInfoMap.Get(LODIndex, SectionIndex);
-
-			if (!NewTask->MaterialBlendModes.IsValidIndex(Section.MaterialIndex))
-			{
-				continue;
-			}
-
-			FSignedDistanceFieldBuildMaterialData& MaterialData = NewTask->MaterialBlendModes[Section.MaterialIndex];
-			MaterialData.bAffectDistanceFieldLighting = Section.bAffectDistanceFieldLighting;
-
-			UMaterialInterface* MaterialInterface = StaticMaterials[Section.MaterialIndex].MaterialInterface;
-			if (MaterialInterface)
-			{
-				MaterialData.BlendMode = MaterialInterface->GetBlendMode();
-				MaterialData.bTwoSided = MaterialInterface->IsTwoSided();
-			}
-		}
+		BuildSignedDistanceFieldBuildSectionData(Mesh, LODIndex, NewTask->SectionData);
 
 		// Nanite overrides source static mesh with a coarse representation. Need to load original data before we build the mesh SDF.
 		if (OptionalSourceMeshData)
@@ -207,7 +187,7 @@ void FCardRepresentationData::CacheDerivedData(const FString& InDDCKey, const IT
 		else if (Mesh->IsNaniteEnabled())
 		{
 			IMeshBuilderModule& MeshBuilderModule = IMeshBuilderModule::GetForPlatform(TargetPlatform);
-			if (!MeshBuilderModule.BuildMeshVertexPositions(Mesh, NewTask->SourceMeshData.TriangleIndices, NewTask->SourceMeshData.VertexPositions))
+			if (!MeshBuilderModule.BuildMeshVertexPositions(Mesh, NewTask->SourceMeshData.TriangleIndices, NewTask->SourceMeshData.VertexPositions, NewTask->SourceMeshData.Sections))
 			{
 				UE_LOG(LogStaticMesh, Error, TEXT("Failed to build static mesh. See previous line(s) for details."));
 			}
@@ -672,7 +652,7 @@ void FCardRepresentationAsyncQueue::Build(FAsyncCardRepresentationTask* Task, FQ
 			Task->SourceMeshData,
 			LODModel,
 			BuildThreadPool,
-			Task->MaterialBlendModes,
+			Task->SectionData,
 			Task->GenerateSource->GetRenderData()->Bounds,
 			Task->GenerateSource->GetRenderData()->LODResources[0].DistanceFieldData,
 			Task->MaxLumenMeshCards,
diff --git a/Engine/Source/Runtime/Engine/Public/DerivedMeshDataTaskUtils.h b/Engine/Source/Runtime/Engine/Public/DerivedMeshDataTaskUtils.h
--- a/Engine/Source/Runtime/Engine/Public/DerivedMeshDataTaskUtils.h
+++ b/Engine/Source/Runtime/Engine/Public/DerivedMeshDataTaskUtils.h
@@ -15,6 +15,7 @@ class FSourceMeshDataForDerivedDataTask
 public:
 	TArray<uint32> TriangleIndices;
 	TArray<FVector3f> VertexPositions;
+	FStaticMeshSectionArray Sections;
 
 	int32 GetNumIndices() const
 	{
diff --git a/Engine/Source/Runtime/Engine/Public/DistanceFieldAtlas.h b/Engine/Source/Runtime/Engine/Public/DistanceFieldAtlas.h
--- a/Engine/Source/Runtime/Engine/Public/DistanceFieldAtlas.h
+++ b/Engine/Source/Runtime/Engine/Public/DistanceFieldAtlas.h
@@ -334,7 +334,7 @@ public:
 	FAsyncDistanceFieldTask();
 
 #if WITH_EDITOR
-	TArray<FSignedDistanceFieldBuildMaterialData> MaterialBlendModes;
+	TArray<FSignedDistanceFieldBuildSectionData> SectionData;
 #endif
 	FSourceMeshDataForDerivedDataTask SourceMeshData;
 	UStaticMesh* StaticMesh;
diff --git a/Engine/Source/Runtime/Engine/Public/MeshCardBuild.h b/Engine/Source/Runtime/Engine/Public/MeshCardBuild.h
--- a/Engine/Source/Runtime/Engine/Public/MeshCardBuild.h
+++ b/Engine/Source/Runtime/Engine/Public/MeshCardBuild.h
@@ -12,7 +12,7 @@
 #include <atomic>
 
 class UStaticMesh;
-class FSignedDistanceFieldBuildMaterialData;
+struct FSignedDistanceFieldBuildSectionData;
 
 class FLumenCardBuildData
 {
@@ -190,7 +190,7 @@ public:
 	bool bSuccess = false;
 
 #if WITH_EDITOR
-	TArray<FSignedDistanceFieldBuildMaterialData> MaterialBlendModes;
+	TArray<FSignedDistanceFieldBuildSectionData> SectionData;
 #endif
 
 	FSourceMeshDataForDerivedDataTask SourceMeshData;
