Replaced double precision members of FDistanceFieldVolumeData and FSparseDistanceFieldMip with single precision floats.

diff --git a/Engine/Source/Developer/MeshUtilities/Private/MeshCardRepresentationUtilities.cpp b/Engine/Source/Developer/MeshUtilities/Private/MeshCardRepresentationUtilities.cpp
--- a/Engine/Source/Developer/MeshUtilities/Private/MeshCardRepresentationUtilities.cpp
+++ b/Engine/Source/Developer/MeshUtilities/Private/MeshCardRepresentationUtilities.cpp
@@ -1089,7 +1089,7 @@ bool FMeshUtilities::GenerateCardRepresentationData(
 		FGenerateCardMeshContext Context(MeshName, EmbreeScene, OutData);
 
 		// Note: must operate on the SDF bounds when available, because SDF generation can expand the mesh's bounds
-		const FBox BuildCardsBounds = DistanceFieldVolumeData && DistanceFieldVolumeData->LocalSpaceMeshBounds.IsValid ? DistanceFieldVolumeData->LocalSpaceMeshBounds : Bounds.GetBox();
+		const FBox BuildCardsBounds = DistanceFieldVolumeData && DistanceFieldVolumeData->LocalSpaceMeshBounds.IsValid ? FBox(DistanceFieldVolumeData->LocalSpaceMeshBounds) : Bounds.GetBox();
 		BuildMeshCards(BuildCardsBounds, Context, MaxLumenMeshCards, OutData);
 
 		MeshRepresentation::DeleteEmbreeScene(EmbreeScene);
diff --git a/Engine/Source/Developer/MeshUtilities/Private/MeshDistanceFieldUtilities.cpp b/Engine/Source/Developer/MeshUtilities/Private/MeshDistanceFieldUtilities.cpp
--- a/Engine/Source/Developer/MeshUtilities/Private/MeshDistanceFieldUtilities.cpp
+++ b/Engine/Source/Developer/MeshUtilities/Private/MeshDistanceFieldUtilities.cpp
@@ -76,11 +76,11 @@ public:
 		const FEmbreeScene& InEmbreeScene,
 		const TArray<FVector3f>* InSampleDirections,
 		float InLocalSpaceTraceDistance,
-		FBox InVolumeBounds,
+		FBox3f InVolumeBounds,
 		float InLocalToVolumeScale,
-		FVector2D InDistanceFieldToVolumeScaleBias,
-		FIntVector InBrickCoordinate,
-		FIntVector InIndirectionSize,
+		FVector2f InDistanceFieldToVolumeScaleBias,
+		FInt32Vector InBrickCoordinate,
+		FInt32Vector InIndirectionSize,
 		bool bInUsePointQuery)
 		:
 		EmbreeScene(InEmbreeScene),
@@ -253,7 +253,7 @@ void FMeshUtilities::GenerateSignedDistanceFieldVolumeData(
 	const FStaticMeshLODResources& LODModel,
 	class FQueuedThreadPool& ThreadPool,
 	const TArray<FSignedDistanceFieldBuildSectionData>& SectionData,
-	const FBoxSphereBounds& Bounds,
+	const FBoxSphereBounds3f& Bounds,
 	float DistanceFieldResolutionScale,
 	bool bGenerateAsIfTwoSided,
 	FDistanceFieldVolumeData& OutData)
@@ -311,12 +311,12 @@ void FMeshUtilities::GenerateSignedDistanceFieldVolumeData(
 		const float VoxelDensity = CVarDensity->GetValueOnAnyThread();
 
 		const float NumVoxelsPerLocalSpaceUnit = VoxelDensity * DistanceFieldResolutionScale;
-		FBox LocalSpaceMeshBounds(Bounds.GetBox());
+		FBox3f LocalSpaceMeshBounds(Bounds.GetBox());
 
 		// Make sure the mesh bounding box has positive extents to handle planes
 		{
-			FVector MeshBoundsCenter = LocalSpaceMeshBounds.GetCenter();
-			FVector MeshBoundsExtent = FVector::Max(LocalSpaceMeshBounds.GetExtent(), FVector(1.0f, 1.0f, 1.0f));
+			FVector3f MeshBoundsCenter = LocalSpaceMeshBounds.GetCenter();
+			FVector3f MeshBoundsExtent = FVector3f::Max(LocalSpaceMeshBounds.GetExtent(), FVector3f::OneVector);
 			LocalSpaceMeshBounds.Min = MeshBoundsCenter - MeshBoundsExtent;
 			LocalSpaceMeshBounds.Max = MeshBoundsCenter + MeshBoundsExtent;
 		}
@@ -326,22 +326,22 @@ void FMeshUtilities::GenerateSignedDistanceFieldVolumeData(
 		// Only expand for two sided meshes as this adds significant Mesh SDF tracing cost
 		if (EmbreeScene.bMostlyTwoSided)
 		{
-			const FVector DesiredDimensions = FVector(LocalSpaceMeshBounds.GetSize() * FVector(NumVoxelsPerLocalSpaceUnit / (float)DistanceField::UniqueDataBrickSize));
-			const FIntVector Mip0IndirectionDimensions = FIntVector(
+			const FVector3f DesiredDimensions = LocalSpaceMeshBounds.GetSize() * FVector3f(NumVoxelsPerLocalSpaceUnit / (float)DistanceField::UniqueDataBrickSize);
+			const FInt32Vector Mip0IndirectionDimensions = FInt32Vector(
 				FMath::Clamp(FMath::RoundToInt(DesiredDimensions.X), 1, MaxNumBlocksOneDim),
 				FMath::Clamp(FMath::RoundToInt(DesiredDimensions.Y), 1, MaxNumBlocksOneDim),
 				FMath::Clamp(FMath::RoundToInt(DesiredDimensions.Z), 1, MaxNumBlocksOneDim));
 
 			const float CentralDifferencingExpandInVoxels = .25f;
-			const FVector TexelObjectSpaceSize = LocalSpaceMeshBounds.GetSize() / FVector(Mip0IndirectionDimensions * DistanceField::UniqueDataBrickSize - FIntVector(2 * CentralDifferencingExpandInVoxels));
+			const FVector3f TexelObjectSpaceSize = LocalSpaceMeshBounds.GetSize() / FVector3f(Mip0IndirectionDimensions * DistanceField::UniqueDataBrickSize - FInt32Vector(2 * CentralDifferencingExpandInVoxels));
 			LocalSpaceMeshBounds = LocalSpaceMeshBounds.ExpandBy(TexelObjectSpaceSize);
 		}
 
 		// The tracing shader uses a Volume space that is normalized by the maximum extent, to keep Volume space within [-1, 1], we must match that behavior when encoding
 		const float LocalToVolumeScale = 1.0f / LocalSpaceMeshBounds.GetExtent().GetMax();
 
-		const FVector DesiredDimensions = FVector(LocalSpaceMeshBounds.GetSize() * FVector(NumVoxelsPerLocalSpaceUnit / (float)DistanceField::UniqueDataBrickSize));
-		const FIntVector Mip0IndirectionDimensions = FIntVector(
+		const FVector3f DesiredDimensions = FVector3f(LocalSpaceMeshBounds.GetSize() * NumVoxelsPerLocalSpaceUnit / (float)DistanceField::UniqueDataBrickSize);
+		const FInt32Vector Mip0IndirectionDimensions = FInt32Vector(
 			FMath::Clamp(FMath::RoundToInt(DesiredDimensions.X), 1, MaxNumBlocksOneDim),
 			FMath::Clamp(FMath::RoundToInt(DesiredDimensions.Y), 1, MaxNumBlocksOneDim),
 			FMath::Clamp(FMath::RoundToInt(DesiredDimensions.Z), 1, MaxNumBlocksOneDim));
@@ -350,22 +350,22 @@ void FMeshUtilities::GenerateSignedDistanceFieldVolumeData(
 
 		for (int32 MipIndex = 0; MipIndex < DistanceField::NumMips; MipIndex++)
 		{
-			const FIntVector IndirectionDimensions = FIntVector(
+			const FInt32Vector IndirectionDimensions = FInt32Vector(
 				FMath::DivideAndRoundUp(Mip0IndirectionDimensions.X, 1 << MipIndex),
 				FMath::DivideAndRoundUp(Mip0IndirectionDimensions.Y, 1 << MipIndex),
 				FMath::DivideAndRoundUp(Mip0IndirectionDimensions.Z, 1 << MipIndex));
 
 			// Expand to guarantee one voxel border for gradient reconstruction using bilinear filtering
-			const FVector TexelObjectSpaceSize = LocalSpaceMeshBounds.GetSize() / FVector(IndirectionDimensions * DistanceField::UniqueDataBrickSize - FIntVector(2 * DistanceField::MeshDistanceFieldObjectBorder));
-			const FBox DistanceFieldVolumeBounds = LocalSpaceMeshBounds.ExpandBy(TexelObjectSpaceSize);
+			const FVector3f TexelObjectSpaceSize = LocalSpaceMeshBounds.GetSize() / FVector3f(IndirectionDimensions * DistanceField::UniqueDataBrickSize - FIntVector(2 * DistanceField::MeshDistanceFieldObjectBorder));
+			const FBox3f DistanceFieldVolumeBounds = LocalSpaceMeshBounds.ExpandBy(TexelObjectSpaceSize);
 
-			const FVector IndirectionVoxelSize = DistanceFieldVolumeBounds.GetSize() / FVector(IndirectionDimensions);
+			const FVector3f IndirectionVoxelSize = DistanceFieldVolumeBounds.GetSize() / FVector3f(IndirectionDimensions);
 			const float IndirectionVoxelRadius = IndirectionVoxelSize.Size();
 
-			const FVector VolumeSpaceDistanceFieldVoxelSize = IndirectionVoxelSize * LocalToVolumeScale / FVector(DistanceField::UniqueDataBrickSize);
+			const FVector3f VolumeSpaceDistanceFieldVoxelSize = IndirectionVoxelSize * LocalToVolumeScale / FVector3f(DistanceField::UniqueDataBrickSize);
 			const float MaxDistanceForEncoding = VolumeSpaceDistanceFieldVoxelSize.Size() * DistanceField::BandSizeInVoxels;
 			const float LocalSpaceTraceDistance = MaxDistanceForEncoding / LocalToVolumeScale;
-			const FVector2D DistanceFieldToVolumeScaleBias(2.0f * MaxDistanceForEncoding, -MaxDistanceForEncoding);
+			const FVector2f DistanceFieldToVolumeScaleBias(2.0f * MaxDistanceForEncoding, -MaxDistanceForEncoding);
 
 			TArray<FSparseMeshDistanceFieldAsyncTask> AsyncTasks;
 			AsyncTasks.Reserve(IndirectionDimensions.X * IndirectionDimensions.Y * IndirectionDimensions.Z / 8);
@@ -383,7 +383,7 @@ void FMeshUtilities::GenerateSignedDistanceFieldVolumeData(
 							DistanceFieldVolumeBounds,
 							LocalToVolumeScale,
 							DistanceFieldToVolumeScaleBias,
-							FIntVector(XIndex, YIndex, ZIndex),
+							FInt32Vector(XIndex, YIndex, ZIndex),
 							IndirectionDimensions,
 							bUsePointQuery);
 					}
@@ -485,10 +485,10 @@ void FMeshUtilities::GenerateSignedDistanceFieldVolumeData(
 			OutMip.NumDistanceFieldBricks = NumBricks;
 
 			// Account for the border voxels we added
-			const FVector VirtualUVMin = FVector(DistanceField::MeshDistanceFieldObjectBorder) / FVector(IndirectionDimensions * DistanceField::UniqueDataBrickSize);
-			const FVector VirtualUVSize = FVector(IndirectionDimensions * DistanceField::UniqueDataBrickSize - FIntVector(2 * DistanceField::MeshDistanceFieldObjectBorder)) / FVector(IndirectionDimensions * DistanceField::UniqueDataBrickSize);
+			const FVector3f VirtualUVMin = FVector3f(DistanceField::MeshDistanceFieldObjectBorder) / FVector3f(IndirectionDimensions * DistanceField::UniqueDataBrickSize);
+			const FVector3f VirtualUVSize = FVector3f(IndirectionDimensions * DistanceField::UniqueDataBrickSize - FInt32Vector(2 * DistanceField::MeshDistanceFieldObjectBorder)) / FVector3f(IndirectionDimensions * DistanceField::UniqueDataBrickSize);
 		
-			const FVector VolumePositionExtent = LocalSpaceMeshBounds.GetExtent() * LocalToVolumeScale;
+			const FVector3f VolumePositionExtent = LocalSpaceMeshBounds.GetExtent() * LocalToVolumeScale;
 
 			// [-VolumePositionExtent, VolumePositionExtent] -> [VirtualUVMin, VirtualUVMin + VirtualUVSize]
 			OutMip.VolumeToVirtualUVScale = VirtualUVSize / (2 * VolumePositionExtent);
diff --git a/Engine/Source/Developer/MeshUtilities/Private/MeshUtilitiesPrivate.h b/Engine/Source/Developer/MeshUtilities/Private/MeshUtilitiesPrivate.h
--- a/Engine/Source/Developer/MeshUtilities/Private/MeshUtilitiesPrivate.h
+++ b/Engine/Source/Developer/MeshUtilities/Private/MeshUtilitiesPrivate.h
@@ -61,7 +61,7 @@ private:
 		const FStaticMeshLODResources& LODModel,
 		class FQueuedThreadPool& ThreadPool,
 		const TArray<FSignedDistanceFieldBuildSectionData>& SectionData,
-		const FBoxSphereBounds& Bounds,
+		const FBoxSphereBounds3f& Bounds,
 		float DistanceFieldResolutionScale,
 		bool bGenerateAsIfTwoSided,
 		FDistanceFieldVolumeData& OutData) override;
diff --git a/Engine/Source/Developer/MeshUtilities/Public/MeshUtilities.h b/Engine/Source/Developer/MeshUtilities/Public/MeshUtilities.h
--- a/Engine/Source/Developer/MeshUtilities/Public/MeshUtilities.h
+++ b/Engine/Source/Developer/MeshUtilities/Public/MeshUtilities.h
@@ -130,7 +130,7 @@ public:
 		const FStaticMeshLODResources& LODModel,
 		class FQueuedThreadPool& ThreadPool,
 		const TArray<FSignedDistanceFieldBuildSectionData>& SectionData,
-		const FBoxSphereBounds& Bounds,
+		const FBoxSphereBounds3f& Bounds,
 		float DistanceFieldResolutionScale,
 		bool bGenerateAsIfTwoSided,
 		class FDistanceFieldVolumeData& OutData) = 0;
diff --git a/Engine/Source/Runtime/Engine/Private/DistanceFieldAtlas.cpp b/Engine/Source/Runtime/Engine/Private/DistanceFieldAtlas.cpp
--- a/Engine/Source/Runtime/Engine/Private/DistanceFieldAtlas.cpp
+++ b/Engine/Source/Runtime/Engine/Private/DistanceFieldAtlas.cpp
@@ -126,7 +126,7 @@ FDistanceFieldAsyncQueue* GDistanceFieldAsyncQueue = NULL;
 #if WITH_EDITOR
 
 // DDC key for distance field data, must be changed when modifying the generation code or data format
-#define DISTANCEFIELD_DERIVEDDATA_VER TEXT("235D7D36-C204-44B7-A44C-D9C67359075C")
+#define DISTANCEFIELD_DERIVEDDATA_VER TEXT("4BD0E980-7212-4093-A14D-2E468CD9FAF3")
 
 FString BuildDistanceFieldDerivedDataKey(const FString& InMeshKey)
 {
@@ -727,7 +727,7 @@ void FDistanceFieldAsyncQueue::Build(FAsyncDistanceFieldTask* Task, FQueuedThrea
 			LODModel,
 			BuildThreadPool,
 			Task->SectionData,
-			Task->GenerateSource->GetRenderData()->Bounds,
+			(FBoxSphereBounds3f)Task->GenerateSource->GetRenderData()->Bounds,
 			Task->DistanceFieldResolutionScale,
 			Task->bGenerateDistanceFieldAsIfTwoSided,
 			*Task->GeneratedVolumeData);
diff --git a/Engine/Source/Runtime/Engine/Public/DistanceFieldAtlas.h b/Engine/Source/Runtime/Engine/Public/DistanceFieldAtlas.h
--- a/Engine/Source/Runtime/Engine/Public/DistanceFieldAtlas.h
+++ b/Engine/Source/Runtime/Engine/Public/DistanceFieldAtlas.h
@@ -208,20 +208,20 @@ class FSparseDistanceFieldMip
 public:
 
 	FSparseDistanceFieldMip() :
-		IndirectionDimensions(FIntVector::ZeroValue),
+		IndirectionDimensions(FInt32Vector::ZeroValue),
 		NumDistanceFieldBricks(0),
-		VolumeToVirtualUVScale(FVector::ZeroVector),
-		VolumeToVirtualUVAdd(FVector::ZeroVector),
-		DistanceFieldToVolumeScaleBias(FVector2D::ZeroVector),
+		VolumeToVirtualUVScale(FVector3f::ZeroVector),
+		VolumeToVirtualUVAdd(FVector3f::ZeroVector),
+		DistanceFieldToVolumeScaleBias(FVector2f::ZeroVector),
 		BulkOffset(0),
 		BulkSize(0)
 	{}
 
-	FIntVector IndirectionDimensions;
+	FInt32Vector IndirectionDimensions;
 	int32 NumDistanceFieldBricks;
-	FVector VolumeToVirtualUVScale;
-	FVector VolumeToVirtualUVAdd;
-	FVector2D DistanceFieldToVolumeScaleBias;
+	FVector3f VolumeToVirtualUVScale;
+	FVector3f VolumeToVirtualUVAdd;
+	FVector2f DistanceFieldToVolumeScaleBias;
 	uint32 BulkOffset;
 	uint32 BulkSize;
 
@@ -250,7 +250,7 @@ class FDistanceFieldVolumeData : public FDeferredCleanupInterface
 public:
 
 	/** Local space bounding box of the distance field volume. */
-	FBox LocalSpaceMeshBounds;
+	FBox3f LocalSpaceMeshBounds;
 
 	/** Whether most of the triangles in the mesh used a two-sided material. */
 	bool bMostlyTwoSided;
diff --git a/Engine/Source/Runtime/Renderer/Private/DistanceFieldObjectManagement.cpp b/Engine/Source/Runtime/Renderer/Private/DistanceFieldObjectManagement.cpp
--- a/Engine/Source/Runtime/Renderer/Private/DistanceFieldObjectManagement.cpp
+++ b/Engine/Source/Runtime/Renderer/Private/DistanceFieldObjectManagement.cpp
@@ -411,7 +411,7 @@ void ProcessPrimitiveUpdate(
 
 				IndicesToUpdateInObjectBuffers.Add(UploadIndex);
 
-				const FBox WorldBounds = DistanceFieldData->LocalSpaceMeshBounds.TransformBy(LocalToWorld);
+				const FBox WorldBounds = ((FBox)DistanceFieldData->LocalSpaceMeshBounds).TransformBy(LocalToWorld);
 
 				if (bIsAddOperation)
 				{
@@ -588,12 +588,12 @@ void FDistanceFieldSceneData::UpdateDistanceFieldObjectBuffers(
 								float SelfShadowBias;
 								PrimitiveSceneProxy->GetDistanceFieldAtlasData(DistanceFieldData, SelfShadowBias);
 
-								const FBox LocalSpaceMeshBounds = DistanceFieldData->LocalSpaceMeshBounds;
+								const FBox3f LocalSpaceMeshBounds = DistanceFieldData->LocalSpaceMeshBounds;
 			
 								const FMatrix LocalToWorld = PrimAndInst.LocalToWorld;
 
 								{
-									const FBox WorldSpaceMeshBounds = LocalSpaceMeshBounds.TransformBy(LocalToWorld);
+									const FBox WorldSpaceMeshBounds = ((FBox)LocalSpaceMeshBounds).TransformBy(LocalToWorld);
 
 									const FLargeWorldRenderPosition AbsoluteWorldPosition(WorldSpaceMeshBounds.GetCenter());
 
@@ -625,9 +625,9 @@ void FDistanceFieldSceneData::UpdateDistanceFieldObjectBuffers(
 
 								// Uniformly scale our Volume space to lie within [-1, 1] at the max extent
 								// This is mirrored in the SDF encoding
-								const FBox::FReal LocalToVolumeScale = 1.0f / LocalSpaceMeshBounds.GetExtent().GetMax();
+								const FBox3f::FReal LocalToVolumeScale = 1.0f / LocalSpaceMeshBounds.GetExtent().GetMax();
 
-								const FMatrix VolumeToWorld = FScaleMatrix(1.0f / LocalToVolumeScale) * FTranslationMatrix(LocalSpaceMeshBounds.GetCenter()) * LocalToWorld;
+								const FMatrix VolumeToWorld = FScaleMatrix(1.0f / LocalToVolumeScale) * FTranslationMatrix((FVector)LocalSpaceMeshBounds.GetCenter()) * LocalToWorld;
 
 								const FLargeWorldRenderPosition WorldPosition(VolumeToWorld.GetOrigin());
 								const FVector TilePositionOffset = WorldPosition.GetTileOffset();
@@ -644,13 +644,13 @@ void FDistanceFieldSceneData::UpdateDistanceFieldObjectBuffers(
 								UploadObjectData[2] = (*(FVector4f*)&WorldToVolumeT.M[1]);
 								UploadObjectData[3] = (*(FVector4f*)&WorldToVolumeT.M[2]);
 
-								const FVector VolumePositionExtent = LocalSpaceMeshBounds.GetExtent() * LocalToVolumeScale;
+								const FVector3f VolumePositionExtent = LocalSpaceMeshBounds.GetExtent() * LocalToVolumeScale;
 
 								// Minimal surface bias which increases chance that ray hit will a surface located between two texels
-								float ExpandSurfaceDistance = (GMeshSDFSurfaceBiasExpand * VolumePositionExtent / FVector(DistanceFieldData->Mips[0].IndirectionDimensions * DistanceField::UniqueDataBrickSize)).Size();
+								float ExpandSurfaceDistance = (GMeshSDFSurfaceBiasExpand * VolumePositionExtent / FVector3f(DistanceFieldData->Mips[0].IndirectionDimensions * DistanceField::UniqueDataBrickSize)).Size();
 
 								const float WSign = DistanceFieldData->bMostlyTwoSided ? -1 : 1;
-								UploadObjectData[4] = FVector4f((FVector3f)VolumePositionExtent, WSign * ExpandSurfaceDistance);
+								UploadObjectData[4] = FVector4f(VolumePositionExtent, WSign * ExpandSurfaceDistance);
 
 								const int32 PrimIdx = PrimAndInst.Primitive->GetIndex();
 								const FPrimitiveBounds& PrimBounds = PrimitiveBounds[PrimIdx];
diff --git a/Engine/Source/Runtime/Renderer/Private/DistanceFieldStreaming.cpp b/Engine/Source/Runtime/Renderer/Private/DistanceFieldStreaming.cpp
--- a/Engine/Source/Runtime/Renderer/Private/DistanceFieldStreaming.cpp
+++ b/Engine/Source/Runtime/Renderer/Private/DistanceFieldStreaming.cpp
@@ -1295,7 +1295,7 @@ void EncodeAssetData(const FDistanceFieldAssetState& AssetState, const int32 Rev
 	const FDistanceFieldAssetMipState& MipState = AssetState.ReversedMips[ReversedMipIndex];
 	const int32 MipIndex = AssetState.BuiltData->Mips.Num() - ReversedMipIndex - 1;
 	const FSparseDistanceFieldMip& MipBuiltData = AssetState.BuiltData->Mips[MipIndex];
-	const FVector2D DistanceFieldToVolumeScaleBias = MipBuiltData.DistanceFieldToVolumeScaleBias;
+	const FVector2f DistanceFieldToVolumeScaleBias = MipBuiltData.DistanceFieldToVolumeScaleBias;
 	const int32 NumMips = AssetState.ReversedMips.Num();
 
 	check(NumMips <= DistanceField::NumMips);
@@ -1319,12 +1319,12 @@ void EncodeAssetData(const FDistanceFieldAssetState& AssetState, const int32 Rev
 	FloatVector0.Z = *(const float*)&IntVector0[2];
 	FloatVector0.W = *(const float*)&IntVector0[3];
 
-	FVector4f VolumeToIndirectionScale = FVector4f((FVector3f)MipBuiltData.VolumeToVirtualUVScale, DistanceFieldToVolumeScaleBias.X);
+	FVector4f VolumeToIndirectionScale = FVector4f(MipBuiltData.VolumeToVirtualUVScale, DistanceFieldToVolumeScaleBias.X);
 	VolumeToIndirectionScale.X *= MipBuiltData.IndirectionDimensions.X;
 	VolumeToIndirectionScale.Y *= MipBuiltData.IndirectionDimensions.Y;
 	VolumeToIndirectionScale.Z *= MipBuiltData.IndirectionDimensions.Z;
 
-	FVector4f VolumeToIndirectionAdd = FVector4f((FVector3f)MipBuiltData.VolumeToVirtualUVAdd, DistanceFieldToVolumeScaleBias.Y);
+	FVector4f VolumeToIndirectionAdd = FVector4f(MipBuiltData.VolumeToVirtualUVAdd, DistanceFieldToVolumeScaleBias.Y);
 	VolumeToIndirectionAdd.X *= MipBuiltData.IndirectionDimensions.X;
 	VolumeToIndirectionAdd.Y *= MipBuiltData.IndirectionDimensions.Y;
 	VolumeToIndirectionAdd.Z *= MipBuiltData.IndirectionDimensions.Z;
