Skeletal Mesh: Create healthy tangents when reconstructing normals/tangents for render and the UVs are non-existent or degenerate.

#rnx

diff --git a/Engine/Source/Developer/MeshUtilities/Private/MeshUtilities.cpp b/Engine/Source/Developer/MeshUtilities/Private/MeshUtilities.cpp
--- a/Engine/Source/Developer/MeshUtilities/Private/MeshUtilities.cpp
+++ b/Engine/Source/Developer/MeshUtilities/Private/MeshUtilities.cpp
@@ -194,11 +194,36 @@ void CalculateTriangleTangentInternal(const FVector3f& VertexPosA, const FVector
 			FPlane4f(0, 0, 0, 1)
 		);
 
-		// Use InverseSlow to catch singular matrices.  Inverse can miss this sometimes.
-		const FMatrix44f TextureToLocal = ParameterToTexture.Inverse() * ParameterToLocal;
+		if (!FMath::IsNearlyZero(ParameterToTexture.Determinant()))
+		{
+			// Use InverseSlow to catch singular matrices.  Inverse can miss this sometimes.
+			const FMatrix44f TextureToLocal = ParameterToTexture.Inverse() * ParameterToLocal;
+
+			OutTangents[0] = (TextureToLocal.TransformVector(FVector3f(1, 0, 0)).GetSafeNormal());
+			OutTangents[1] = (TextureToLocal.TransformVector(FVector3f(0, 1, 0)).GetSafeNormal());
+		}
+		else
+		{
+			// Use the Duff & Frisvad algorithm (see Duff 2017 in JCGT) to construct a consistent tangent from a
+			// single normal vector. 
+			if (Normal.Z < 0.0f)
+			{
+				const float A = 1.0f / (1.0f - Normal.Z);
+				const float B = Normal.X * Normal.Y * A;
+
+				OutTangents[0] = FVector3f(1.0f - Normal.X * Normal.X * A, -B, Normal.X);
+				OutTangents[1] = FVector3f(B, Normal.Y * Normal.Y * A - 1.0f, -Normal.Y);
+			}
+			else
+			{
+				const float A = 1.0f / (1.0f + Normal.Z);
+				const float B = -Normal.X * Normal.Y * A;
+
+				OutTangents[0] = FVector3f(1.0f - Normal.X * Normal.X * A, B, -Normal.X);
+				OutTangents[1] = FVector3f(B, 1.0f - Normal.Y * Normal.Y * A, -Normal.Y);
+			}
+		}
 
-		OutTangents[0] = (TextureToLocal.TransformVector(FVector3f(1, 0, 0)).GetSafeNormal());
-		OutTangents[1] = (TextureToLocal.TransformVector(FVector3f(0, 1, 0)).GetSafeNormal());
 		OutTangents[2] = (Normal);
 
 		FVector3f::CreateOrthonormalBasis(
@@ -3855,17 +3880,17 @@ public:
 				CornerNormal[CornerIndex].Normalize();
 				if (!bUseMikktSpace)
 				{
-					CornerTangentX[CornerIndex].Normalize();
-					CornerTangentY[CornerIndex].Normalize();
-
-					// Gram-Schmidt orthogonalization
-					CornerTangentY[CornerIndex] -= CornerTangentX[CornerIndex] * (CornerTangentX[CornerIndex] | CornerTangentY[CornerIndex]);
-					CornerTangentY[CornerIndex].Normalize();
-
-					CornerTangentX[CornerIndex] -= CornerNormal[CornerIndex] * (CornerNormal[CornerIndex] | CornerTangentX[CornerIndex]);
-					CornerTangentX[CornerIndex].Normalize();
-					CornerTangentY[CornerIndex] -= CornerNormal[CornerIndex] * (CornerNormal[CornerIndex] | CornerTangentY[CornerIndex]);
-					CornerTangentY[CornerIndex].Normalize();
+					// If Skeletal_ComputeTriangleTangents somehow didn't manage to create any tangents for this point and its neighbors,
+					// just feed the orthonormal basis computation something. 
+					if (CornerTangentX[CornerIndex].IsNearlyZero())
+					{
+						CornerTangentX[CornerIndex] = FVector3f::XAxisVector;
+					}
+					if (CornerTangentY[CornerIndex].IsNearlyZero())
+					{
+						CornerTangentY[CornerIndex] = FVector3f::YAxisVector;
+					}
+					FVector3f::CreateOrthonormalBasis(CornerTangentX[CornerIndex], CornerTangentY[CornerIndex], CornerNormal[CornerIndex]);
 				}
 			}
 
