[SkeletalMesh to StaticMesh] MaterialIndices fixes.

#rnx

diff --git a/Engine/Source/Developer/MeshUtilities/Private/MeshUtilities.cpp b/Engine/Source/Developer/MeshUtilities/Private/MeshUtilities.cpp
--- a/Engine/Source/Developer/MeshUtilities/Private/MeshUtilities.cpp
+++ b/Engine/Source/Developer/MeshUtilities/Private/MeshUtilities.cpp
@@ -431,6 +431,8 @@ static void SkinnedMeshToRawMeshes(USkinnedMeshComponent* InSkinnedMeshComponent
 				}
 			}
 		}
+
+		RawMesh.CompactMaterialIndices();
 	}
 
 	ProcessMaterials<USkinnedMeshComponent>(InSkinnedMeshComponent, InPackageName, OutMaterials);
@@ -729,7 +731,14 @@ UStaticMesh* FMeshUtilities::ConvertMeshesToStaticMesh(const TArray<UMeshCompone
 				int32 SectionIndex = 0;
 				for (int32 UniqueMaterialIndex : UniqueMaterialIndices)
 				{
-					StaticMesh->GetSectionInfoMap().Set(RawMeshLODIndex, SectionIndex, FMeshSectionInfo(UniqueMaterialIndex));
+					if (RawMesh.MaterialIndexToImportIndex.IsValidIndex(SectionIndex))
+					{
+						StaticMesh->GetSectionInfoMap().Set(RawMeshLODIndex, SectionIndex, FMeshSectionInfo(RawMesh.MaterialIndexToImportIndex[SectionIndex]));
+					}
+					else
+					{
+						StaticMesh->GetSectionInfoMap().Set(RawMeshLODIndex, SectionIndex, FMeshSectionInfo(UniqueMaterialIndex));
+					}
 					SectionIndex++;
 				}
 			}
