Lumen - fixed mesh card generation when first card fails and we stop the entire process instead of trying to generate another one with a different near plane (in a different spot).

diff --git a/Engine/Source/Developer/MeshUtilities/Private/MeshCardRepresentationUtilities.cpp b/Engine/Source/Developer/MeshUtilities/Private/MeshCardRepresentationUtilities.cpp
--- a/Engine/Source/Developer/MeshUtilities/Private/MeshCardRepresentationUtilities.cpp
+++ b/Engine/Source/Developer/MeshUtilities/Private/MeshCardRepresentationUtilities.cpp
@@ -866,11 +866,10 @@ void BuildSurfelClusters(const FBox& MeshBounds, const FGenerateCardMeshContext&
 
 			FSurfelCluster TempCluster;
 			BuildCluster(/*NearPlane*/ 0, SurfelScenePerDirection, SurfelAssignedToAnyCluster, TempCluster);
-			bool bCanAddCluster = TempCluster.IsValid(ClusteringParams);
-			if (bCanAddCluster)
+
+			if (TempCluster.IsValid(ClusteringParams))
 			{
 				CommitCluster(Clusters, SurfelAssignedToAnyCluster, TempCluster);
-				bCanAddCluster = true;
 			}
 
 			// Assume that two sided is foliage and revert to a simpler box projection
@@ -878,6 +877,7 @@ void BuildSurfelClusters(const FBox& MeshBounds, const FGenerateCardMeshContext&
 			{
 				FSurfelCluster BestCluster;
 
+				bool bCanAddCluster = true;
 				while (bCanAddCluster)
 				{
 					BestCluster.Init(-1);
diff --git a/Engine/Source/Runtime/Engine/Private/MeshCardRepresentation.cpp b/Engine/Source/Runtime/Engine/Private/MeshCardRepresentation.cpp
--- a/Engine/Source/Runtime/Engine/Private/MeshCardRepresentation.cpp
+++ b/Engine/Source/Runtime/Engine/Private/MeshCardRepresentation.cpp
@@ -101,7 +101,7 @@ FCardRepresentationAsyncQueue* GCardRepresentationAsyncQueue = NULL;
 #if WITH_EDITOR
 
 // DDC key for card representation data, must be changed when modifying the generation code or data format
-#define CARDREPRESENTATION_DERIVEDDATA_VER TEXT("7DD7930F-6ED7-4CF1-BE60-E9819779DBAF")
+#define CARDREPRESENTATION_DERIVEDDATA_VER TEXT("B58BB068-C040-4A1D-BC7C-A06D36FE3B1D")
 
 FString BuildCardRepresentationDerivedDataKey(const FString& InMeshKey, int32 MaxLumenMeshCards)
 {
