Fixing issues with some legacy cook stat support not carrying over, plus some other cook stat improvements.
#rnx

diff --git a/Engine/Source/Developer/DerivedDataCache/Public/DerivedDataCacheUsageStats.h b/Engine/Source/Developer/DerivedDataCache/Public/DerivedDataCacheUsageStats.h
--- a/Engine/Source/Developer/DerivedDataCache/Public/DerivedDataCacheUsageStats.h
+++ b/Engine/Source/Developer/DerivedDataCache/Public/DerivedDataCacheUsageStats.h
@@ -224,7 +224,7 @@ protected:
 struct FDerivedDataCacheResourceStat
 {
 public:
-	FDerivedDataCacheResourceStat(FString InAssetType = TEXT("None"), bool bIsGameThreadTime = 0.0, double InLoadTimeSec = 0.0, double InLoadSizeMB = 0.0, int64 InAssetsLoaded = 0, double InBuildTimeSec = 0.0, double InBuildSizeMB = 0.0, int64 InAssetsBuilt = 0) :
+	FDerivedDataCacheResourceStat(FString InAssetType = TEXT("None"), bool bIsGameThreadTime = false, double InLoadTimeSec = 0.0, double InLoadSizeMB = 0.0, int64 InAssetsLoaded = 0, double InBuildTimeSec = 0.0, double InBuildSizeMB = 0.0, int64 InAssetsBuilt = 0) :
 		AssetType(MoveTemp(InAssetType)),
 		LoadTimeSec(InLoadTimeSec),
 		LoadSizeMB(InLoadSizeMB),
diff --git a/Engine/Source/Developer/MeshUtilities/Private/MeshDistanceFieldUtilities.cpp b/Engine/Source/Developer/MeshUtilities/Private/MeshDistanceFieldUtilities.cpp
--- a/Engine/Source/Developer/MeshUtilities/Private/MeshDistanceFieldUtilities.cpp
+++ b/Engine/Source/Developer/MeshUtilities/Private/MeshDistanceFieldUtilities.cpp
@@ -10,6 +10,18 @@
 #include "DistanceFieldAtlas.h"
 #include "MeshRepresentationCommon.h"
 #include "Async/ParallelFor.h"
+#include "ProfilingDebugging/CookStats.h"
+
+#if ENABLE_COOK_STATS
+namespace DistanceFieldUtilityCookStats
+{
+	FCookStats::FDDCResourceUsageStats UsageStats;
+	static FCookStatsManager::FAutoRegisterCallback RegisterCookStats([](FCookStatsManager::AddStatFuncRef AddStat)
+	{
+		UsageStats.LogStats(AddStat, TEXT("DistanceField.Usage"), TEXT(""));
+	});
+}
+#endif
 
 #if USE_EMBREE
 
@@ -122,6 +134,8 @@ int32 DebugZ = 0;
 void FSparseMeshDistanceFieldAsyncTask::DoWork()
 {
 	TRACE_CPUPROFILER_EVENT_SCOPE(FSparseMeshDistanceFieldAsyncTask::DoWork);
+	COOK_STAT(auto Timer = DistanceFieldUtilityCookStats::UsageStats.TimeSyncWork());
+	COOK_STAT(Timer.TrackCyclesOnly());
 
 	const FVector IndirectionVoxelSize = VolumeBounds.GetSize() / FVector(IndirectionSize);
 	const FVector DistanceFieldVoxelSize = IndirectionVoxelSize / FVector(DistanceField::UniqueDataBrickSize);
@@ -259,6 +273,8 @@ void FMeshUtilities::GenerateSignedDistanceFieldVolumeData(
 	FDistanceFieldVolumeData& OutData)
 {
 	TRACE_CPUPROFILER_EVENT_SCOPE(GenerateSignedDistanceFieldVolumeData);
+	COOK_STAT(auto Timer = DistanceFieldUtilityCookStats::UsageStats.TimeSyncWork());
+	COOK_STAT(Timer.TrackCyclesOnly());
 
 	if (DistanceFieldResolutionScale > 0)
 	{
@@ -395,6 +411,8 @@ void FMeshUtilities::GenerateSignedDistanceFieldVolumeData(
 			if (bMultiThreaded)
 			{
 				EParallelForFlags Flags = EParallelForFlags::BackgroundPriority | EParallelForFlags::Unbalanced;
+				COOK_STAT(auto AsyncTimer = DistanceFieldUtilityCookStats::UsageStats.TimeAsyncWait());
+				COOK_STAT(AsyncTimer.TrackCyclesOnly());
 
 				ParallelForTemplate(
 					TEXT("GenerateSignedDistanceFieldVolumeData.PF"),
diff --git a/Engine/Source/Runtime/Core/Public/ProfilingDebugging/CookStats.h b/Engine/Source/Runtime/Core/Public/ProfilingDebugging/CookStats.h
--- a/Engine/Source/Runtime/Core/Public/ProfilingDebugging/CookStats.h
+++ b/Engine/Source/Runtime/Core/Public/ProfilingDebugging/CookStats.h
@@ -302,7 +302,7 @@ public:
 		/** Starts the time, tracks the underlying stat it will update. Use bCyclesOnly if you need to track additional cycles in a different scope without adding to the call count. */
 		explicit FScopedStatsCounter(CallStats& InStats)
 			: Stats(InStats)
-			, StartTime(FPlatformTime::Seconds())
+			, StartTime(FPlatformTime::Cycles64())
 			, bIsInGameThread(IsInGameThread())
 		{
 		}
@@ -326,8 +326,9 @@ public:
 		{
 			if (!bCanceled)
 			{
-				// Can't safely use FPlatformTime::Cycles() because we might be timing long durations.
-				const int64 CyclesUsed = int64((FPlatformTime::Seconds() - StartTime) / FPlatformTime::GetSecondsPerCycle());
+				uint64 CurrentTime = FPlatformTime::Cycles64();
+				const uint64 CyclesUsed = CurrentTime - StartTime;
+
 				Stats.Accumulate(HitOrMiss, CallStats::EStatType::Cycles, CyclesUsed, bIsInGameThread);
 				if (!bTrackCyclesOnly)
 				{
@@ -373,7 +374,7 @@ public:
 
 	private:
 		CallStats& Stats;
-		double StartTime;
+		uint64 StartTime;
 		int64 BytesProcessed = 0;
 		bool bIsInGameThread;
 		bool bCanceled = false;
diff --git a/Engine/Source/Runtime/Engine/Private/Animation/AnimCompressionDerivedData.cpp b/Engine/Source/Runtime/Engine/Private/Animation/AnimCompressionDerivedData.cpp
--- a/Engine/Source/Runtime/Engine/Private/Animation/AnimCompressionDerivedData.cpp
+++ b/Engine/Source/Runtime/Engine/Private/Animation/AnimCompressionDerivedData.cpp
@@ -58,10 +58,6 @@ bool FDerivedDataAnimationCompression::Build( TArray<uint8>& OutDataArray )
 		return false;
 	}
 
-	// Update UsageStats only when not running on the game thread since that thread times this at a higher level.
-	COOK_STAT(auto Timer = UE::Anim::AnimSequenceCookStats::UsageStats.TimeSyncWork());
-	COOK_STAT(IsInGameThread() ? Timer.Cancel() : Timer.TrackCyclesOnly());
-
 	SCOPE_CYCLE_COUNTER(STAT_AnimCompressionDerivedData);
 	UE_LOG(LogAnimationCompression, Log, TEXT("Building Anim DDC data for %s"), *DataToCompress.FullName);
 
diff --git a/Engine/Source/Runtime/Engine/Private/Animation/AnimSequence.cpp b/Engine/Source/Runtime/Engine/Private/Animation/AnimSequence.cpp
--- a/Engine/Source/Runtime/Engine/Private/Animation/AnimSequence.cpp
+++ b/Engine/Source/Runtime/Engine/Private/Animation/AnimSequence.cpp
@@ -4586,7 +4586,8 @@ FString UAnimSequence::CreateDerivedDataKeyString(const ITargetPlatform* TargetP
 FIoHash UAnimSequence::BeginCacheDerivedData(const ITargetPlatform* TargetPlatform)
 {
 	QUICK_SCOPE_CYCLE_COUNTER(STAT_AnimSequence_BeginCacheDerivedData);
-	
+	COOK_STAT(auto Timer = UE::Anim::AnimSequenceCookStats::UsageStats.TimeSyncWork());
+
 	check(TargetPlatform != nullptr);
 	check(IsInGameThread());
 	check(!FPlatformProperties::RequiresCookedData());
@@ -4677,11 +4678,8 @@ FIoHash UAnimSequence::BeginCacheDerivedData(const ITargetPlatform* TargetPlatfo
 	{		
 		// Data does not exist, need to build it.
 		const FCompressibleAnimRef CompressibleData = MakeShared<FCompressibleAnimData, ESPMode::ThreadSafe>(this, bPerformFrameStripping, TargetPlatform);
-		
-		COOK_STAT(auto Timer = UE::Anim::AnimSequenceCookStats::UsageStats.TimeSyncWork());
 		check(CacheTasksByKeyHash.Find(KeyHash) == nullptr);
 		CacheTasksByKeyHash.Emplace(KeyHash, MakePimpl<UE::Anim::FAnimationSequenceAsyncCacheTask>(KeyHash, CompressibleData, TargetData, *this, TargetPlatform));		
-		COOK_STAT(Timer.TrackCyclesOnly());
 	}
 
 	// The compiling manager provides throttling, notification manager, etc... for the asset being built.
diff --git a/Engine/Source/Runtime/Engine/Private/Animation/AnimationCompressionDerivedData.cpp b/Engine/Source/Runtime/Engine/Private/Animation/AnimationCompressionDerivedData.cpp
--- a/Engine/Source/Runtime/Engine/Private/Animation/AnimationCompressionDerivedData.cpp
+++ b/Engine/Source/Runtime/Engine/Private/Animation/AnimationCompressionDerivedData.cpp
@@ -18,10 +18,10 @@ namespace UE::Anim
 #if ENABLE_COOK_STATS
 	namespace AnimSequenceCookStats
 	{
-		FDerivedDataUsageStats UsageStats;
+		FCookStats::FDDCResourceUsageStats UsageStats;
 		static FCookStatsManager::FAutoRegisterCallback RegisterCookStats([](FCookStatsManager::AddStatFuncRef AddStat)
 		{
-			UsageStats.LogStats(AddStat, TEXT("AnimSequence.Usage"), TEXT(""));
+			UsageStats.LogStats(AddStat, TEXT("AnimationSequence.Usage"), TEXT(""));
 		});
 	}
 #endif
diff --git a/Engine/Source/Runtime/Engine/Private/Animation/AnimationCompressionDerivedData.h b/Engine/Source/Runtime/Engine/Private/Animation/AnimationCompressionDerivedData.h
--- a/Engine/Source/Runtime/Engine/Private/Animation/AnimationCompressionDerivedData.h
+++ b/Engine/Source/Runtime/Engine/Private/Animation/AnimationCompressionDerivedData.h
@@ -28,7 +28,7 @@ namespace UE::Anim
 #if ENABLE_COOK_STATS
 	namespace AnimSequenceCookStats
 	{
-		extern FDerivedDataUsageStats UsageStats;
+		extern FCookStats::FDDCResourceUsageStats UsageStats;
 	}
 #endif
 
diff --git a/Engine/Source/Runtime/Engine/Private/MeshCardRepresentation.cpp b/Engine/Source/Runtime/Engine/Private/MeshCardRepresentation.cpp
--- a/Engine/Source/Runtime/Engine/Private/MeshCardRepresentation.cpp
+++ b/Engine/Source/Runtime/Engine/Private/MeshCardRepresentation.cpp
@@ -648,6 +648,8 @@ void FCardRepresentationAsyncQueue::BlockUntilAllBuildsComplete()
 void FCardRepresentationAsyncQueue::Build(FAsyncCardRepresentationTask* Task, FQueuedThreadPool& BuildThreadPool)
 {
 #if WITH_EDITOR
+	COOK_STAT(auto Timer = CardRepresentationCookStats::UsageStats.TimeSyncWork());
+	COOK_STAT(Timer.TrackCyclesOnly());
 	
 	// Editor 'force delete' can null any UObject pointers which are seen by reference collecting (eg UProperty or serialized)
 	if (Task->StaticMesh && Task->GenerateSource)
diff --git a/Engine/Source/Runtime/Engine/Private/ShaderCompiler/ShaderCompilerStats.cpp b/Engine/Source/Runtime/Engine/Private/ShaderCompiler/ShaderCompilerStats.cpp
--- a/Engine/Source/Runtime/Engine/Private/ShaderCompiler/ShaderCompilerStats.cpp
+++ b/Engine/Source/Runtime/Engine/Private/ShaderCompiler/ShaderCompilerStats.cpp
@@ -396,7 +396,7 @@ void FShaderCompilerStats::WriteStatSummary()
 			double EffectiveParallelization = TotalThreadTimeForAllShaders / TotalTimeAtLeastOneJobWasInFlight;
 			if (Counters.DistributedJobBatchesSeen == 0)
 			{
-				UE_LOG(LogShaderCompilers, Display, TEXT("Effective parallelization: %.2f (times faster than compiling all shaders on one thread). Compare with number of workers: %d"), EffectiveParallelization, GShaderCompilingManager->GetNumLocalWorkers());
+				UE_LOG(LogShaderCompilers, Display, TEXT("Effective parallelization: %.2f (times faster than compiling all shaders on one thread). Compare with number of workers: %d - %f"), EffectiveParallelization, GShaderCompilingManager->GetNumLocalWorkers(), EffectiveParallelization/(double)GShaderCompilingManager->GetNumLocalWorkers());
 			}
 			else
 			{
diff --git a/Engine/Source/Runtime/Engine/Private/TextureDerivedDataTask.cpp b/Engine/Source/Runtime/Engine/Private/TextureDerivedDataTask.cpp
--- a/Engine/Source/Runtime/Engine/Private/TextureDerivedDataTask.cpp
+++ b/Engine/Source/Runtime/Engine/Private/TextureDerivedDataTask.cpp
@@ -35,6 +35,18 @@
 #include "TextureFormatManager.h"
 #include "VT/VirtualTextureChunkDDCCache.h"
 #include "VT/VirtualTextureDataBuilder.h"
+#include "ProfilingDebugging/CookStats.h"
+
+#if ENABLE_COOK_STATS
+namespace TextureCookStats
+{
+	static FCookStats::FDDCResourceUsageStats TaskUsageStats;
+	static FCookStatsManager::FAutoRegisterCallback RegisterTaskCookStats([](FCookStatsManager::AddStatFuncRef AddStat)
+	{
+		TextureCookStats::TaskUsageStats.LogStats(AddStat, TEXT("Texture.Usage"), TEXT("Task"));
+	});
+}
+#endif
 
 static TAutoConsoleVariable<int32> CVarVTValidateCompressionOnLoad(
 	TEXT("r.VT.ValidateCompressionOnLoad"),
@@ -1982,6 +1994,7 @@ void FTextureCacheDerivedDataWorker::DoWork()
 	using namespace UE::DerivedData;
 
 	TRACE_CPUPROFILER_EVENT_SCOPE(FTextureCacheDerivedDataWorker::DoWork);
+	COOK_STAT(auto Timer = TextureCookStats::TaskUsageStats.TimeSyncWork());
 
 	const bool bAllowAsyncBuild = EnumHasAnyFlags(CacheFlags, ETextureCacheFlags::AllowAsyncBuild);
 	const bool bAllowAsyncLoading = EnumHasAnyFlags(CacheFlags, ETextureCacheFlags::AllowAsyncLoading);
@@ -2116,6 +2129,10 @@ void FTextureCacheDerivedDataWorker::Finalize()
 	// --	always with a lighting build, as the async light/shadowmap tasks will disallow async builds. 
 	// --	if the texture compiler cvar disallows async texture compilation  "Editor.AsyncTextureCompilation 0"
 
+	TRACE_CPUPROFILER_EVENT_SCOPE(FTextureCacheDerivedDataWorker::DoWork);
+	COOK_STAT(auto Timer = TextureCookStats::TaskUsageStats.TimeSyncWork());
+	COOK_STAT(Timer.TrackCyclesOnly());
+
 	if ( bTriedAndFailed )
 	{
 		UE_LOG(LogTexture, Warning, TEXT("Texture build failed for %s.  Will not retry in Finalize."), *TexturePathName);
