[DistanceField] Remove unused kDOP code path from MeshRepresentationCommon.

diff --git a/Engine/Source/Developer/MeshUtilities/Private/MeshDistanceFieldUtilities.cpp b/Engine/Source/Developer/MeshUtilities/Private/MeshDistanceFieldUtilities.cpp
--- a/Engine/Source/Developer/MeshUtilities/Private/MeshDistanceFieldUtilities.cpp
+++ b/Engine/Source/Developer/MeshUtilities/Private/MeshDistanceFieldUtilities.cpp
@@ -290,8 +290,6 @@ void FMeshUtilities::GenerateSignedDistanceFieldVolumeData(
 			bIncludeTranslucentTriangles,
 			EmbreeScene);
 
-		check(EmbreeScene.bUseEmbree);
-
 		// If Embree setup fails, there will be no scene to operate on. Early out.
 		if (!EmbreeScene.EmbreeScene)
 		{
diff --git a/Engine/Source/Developer/MeshUtilities/Private/MeshRepresentationCommon.cpp b/Engine/Source/Developer/MeshUtilities/Private/MeshRepresentationCommon.cpp
--- a/Engine/Source/Developer/MeshUtilities/Private/MeshRepresentationCommon.cpp
+++ b/Engine/Source/Developer/MeshUtilities/Private/MeshRepresentationCommon.cpp
@@ -140,12 +140,7 @@ void MeshRepresentation::SetupEmbreeScene(
 
 	const FStaticMeshSectionArray& Sections = SourceMeshData.IsValid() ? SourceMeshData.Sections : LODModel.Sections;
 
-	TArray<FkDOPBuildCollisionTriangle<uint32> > BuildTriangles;
-
 #if USE_EMBREE
-	EmbreeScene.bUseEmbree = true;
-
-	if (EmbreeScene.bUseEmbree)
 	{
 		EmbreeScene.EmbreeDevice = rtcNewDevice(nullptr);
 		rtcSetDeviceErrorFunction(EmbreeScene.EmbreeDevice, EmbreeErrorFunc, nullptr);
@@ -305,7 +300,6 @@ void MeshRepresentation::SetupEmbreeScene(
 			}
 		}
 
-		if (EmbreeScene.bUseEmbree)
 		{
 			EmbreeIndices[FilteredTriangleIndex * 3 + 0] = I0;
 			EmbreeIndices[FilteredTriangleIndex * 3 + 1] = I1;
@@ -320,19 +314,9 @@ void MeshRepresentation::SetupEmbreeScene(
 			Desc.ElementIndex = bGenerateAsIfTwoSided || bTriangleIsTwoSided ? 1 : 0;
 			EmbreeScene.Geometry.TriangleDescs.Add(Desc);
 		}
-		else
-		{
-			BuildTriangles.Add(FkDOPBuildCollisionTriangle<uint32>(
-				// Store bGenerateAsIfTwoSided in material index
-				bGenerateAsIfTwoSided || bTriangleIsTwoSided ? 1 : 0,
-				FVector(V0),
-				FVector(V1),
-				FVector(V2)));
-		}
 	}
 
 #if USE_EMBREE
-	if (EmbreeScene.bUseEmbree)
 	{
 		RTCGeometry Geometry = rtcNewGeometry(EmbreeScene.EmbreeDevice, RTC_GEOMETRY_TYPE_TRIANGLE);
 		EmbreeScene.Geometry.InternalGeometry = Geometry;
@@ -362,11 +346,7 @@ void MeshRepresentation::SetupEmbreeScene(
 			return;
 		}
 	}
-	else
 #endif
-	{
-		EmbreeScene.kDopTree.Build(BuildTriangles);
-	}
 
 	// bMostlyTwoSided
 	{
@@ -395,10 +375,7 @@ void MeshRepresentation::SetupEmbreeScene(
 void MeshRepresentation::DeleteEmbreeScene(FEmbreeScene& EmbreeScene)
 {
 #if USE_EMBREE
-	if (EmbreeScene.bUseEmbree)
-	{
-		rtcReleaseScene(EmbreeScene.EmbreeScene);
-		rtcReleaseDevice(EmbreeScene.EmbreeDevice);
-	}
+	rtcReleaseScene(EmbreeScene.EmbreeScene);
+	rtcReleaseDevice(EmbreeScene.EmbreeDevice);
 #endif
 }
\ No newline at end of file
diff --git a/Engine/Source/Developer/MeshUtilities/Private/MeshRepresentationCommon.h b/Engine/Source/Developer/MeshUtilities/Private/MeshRepresentationCommon.h
--- a/Engine/Source/Developer/MeshUtilities/Private/MeshRepresentationCommon.h
+++ b/Engine/Source/Developer/MeshUtilities/Private/MeshRepresentationCommon.h
@@ -3,7 +3,6 @@
 #pragma once
 
 #include "MeshUtilities.h"
-#include "kDOP.h"
 
 #if USE_EMBREE
 	#include <embree3/rtcore.h>
@@ -16,48 +15,6 @@
 
 class FSourceMeshDataForDerivedDataTask;
 
-class FMeshBuildDataProvider
-{
-public:
-
-	/** Initialization constructor. */
-	FMeshBuildDataProvider(
-		const TkDOPTree<const FMeshBuildDataProvider, uint32>& InkDopTree) :
-		kDopTree(InkDopTree)
-	{}
-
-	// kDOP data provider interface.
-
-	FORCEINLINE const TkDOPTree<const FMeshBuildDataProvider, uint32>& GetkDOPTree(void) const
-	{
-		return kDopTree;
-	}
-
-	FORCEINLINE const FMatrix& GetLocalToWorld(void) const
-	{
-		return FMatrix::Identity;
-	}
-
-	FORCEINLINE const FMatrix& GetWorldToLocal(void) const
-	{
-		return FMatrix::Identity;
-	}
-
-	FORCEINLINE FMatrix GetLocalToWorldTransposeAdjoint(void) const
-	{
-		return FMatrix::Identity;
-	}
-
-	FORCEINLINE float GetDeterminant(void) const
-	{
-		return 1.0f;
-	}
-
-private:
-
-	const TkDOPTree<const FMeshBuildDataProvider, uint32>& kDopTree;
-};
-
 struct FEmbreeTriangleDesc
 {
 	int16 ElementIndex;
@@ -81,17 +38,12 @@ struct FEmbreeGeometry
 class FEmbreeScene
 {
 public:
-	bool bUseEmbree = false;
 	int32 NumIndices = 0;
 	bool bMostlyTwoSided = false;
 
-	// Embree
 	RTCDevice EmbreeDevice = nullptr;
 	RTCScene EmbreeScene = nullptr;
 	FEmbreeGeometry Geometry;
-
-	// DOP tree fallback
-	TkDOPTree<const FMeshBuildDataProvider, uint32> kDopTree;
 };
 
 #if USE_EMBREE
