[DistanceFields] Refactored GenerateSignedDistanceFieldVolumeData(...) to move most of the logic to a separate function.
- this will enable reusing logic in multiple code paths in upcoming changes.

diff --git a/Engine/Source/Developer/MeshUtilities/Private/MeshDistanceFieldUtilities.cpp b/Engine/Source/Developer/MeshUtilities/Private/MeshDistanceFieldUtilities.cpp
--- a/Engine/Source/Developer/MeshUtilities/Private/MeshDistanceFieldUtilities.cpp
+++ b/Engine/Source/Developer/MeshUtilities/Private/MeshDistanceFieldUtilities.cpp
@@ -261,271 +261,283 @@ void FSparseMeshDistanceFieldAsyncTask::DoWork()
 	}
 }
 
-void FMeshUtilities::GenerateSignedDistanceFieldVolumeData(
+static void BuildSignedDistanceField(
 	FString MeshName,
-	const FSourceMeshDataForDerivedDataTask& SourceMeshData,
-	const FStaticMeshLODResources& LODModel,
-	class FQueuedThreadPool& ThreadPool,
-	const TArray<FSignedDistanceFieldBuildSectionData>& SectionData,
+	const FEmbreeScene& EmbreeScene,
 	const FBoxSphereBounds3f& Bounds,
 	float DistanceFieldResolutionScale,
-	bool bGenerateAsIfTwoSided,
 	FDistanceFieldVolumeData& OutData)
 {
-	TRACE_CPUPROFILER_EVENT_SCOPE(GenerateSignedDistanceFieldVolumeData);
-	COOK_STAT(auto Timer = DistanceFieldUtilityCookStats::UsageStats.TimeSyncWork());
-	COOK_STAT(Timer.TrackCyclesOnly());
+	// Whether to use an Embree Point Query to compute the closest unsigned distance.  Rays will only be traced to determine backfaces visible for sign.
+	const bool bUsePointQuery = true;
 
-	if (DistanceFieldResolutionScale > 0)
+	TArray<FVector3f> SampleDirections;
 	{
-		const double StartTime = FPlatformTime::Seconds();
-		const bool bIncludeTranslucentTriangles = false;
-
-		FEmbreeScene EmbreeScene;
-		MeshRepresentation::SetupEmbreeScene(MeshName,
-			SourceMeshData,
-			LODModel,
-			SectionData,
-			bGenerateAsIfTwoSided,
-			bIncludeTranslucentTriangles,
-			EmbreeScene);
-
-		// If Embree setup fails, there will be no scene to operate on. Early out.
-		if (!EmbreeScene.EmbreeScene)
-		{
-			return;
-		}
-
-		// Whether to use an Embree Point Query to compute the closest unsigned distance.  Rays will only be traced to determine backfaces visible for sign.
-		const bool bUsePointQuery = true;
+		const int32 NumVoxelDistanceSamples = bUsePointQuery ? 49 : 576;
+		FRandomStream RandomStream(0);
+		MeshUtilities::GenerateStratifiedUniformHemisphereSamples(NumVoxelDistanceSamples, RandomStream, SampleDirections);
+		TArray<FVector3f> OtherHemisphereSamples;
+		MeshUtilities::GenerateStratifiedUniformHemisphereSamples(NumVoxelDistanceSamples, RandomStream, OtherHemisphereSamples);
 
-		TArray<FVector3f> SampleDirections;
+		for (int32 i = 0; i < OtherHemisphereSamples.Num(); i++)
 		{
-			const int32 NumVoxelDistanceSamples = bUsePointQuery ? 49 : 576;
-			FRandomStream RandomStream(0);
-			MeshUtilities::GenerateStratifiedUniformHemisphereSamples(NumVoxelDistanceSamples, RandomStream, SampleDirections);
-			TArray<FVector3f> OtherHemisphereSamples;
-			MeshUtilities::GenerateStratifiedUniformHemisphereSamples(NumVoxelDistanceSamples, RandomStream, OtherHemisphereSamples);
-
-			for (int32 i = 0; i < OtherHemisphereSamples.Num(); i++)
-			{
-				FVector3f Sample = OtherHemisphereSamples[i];
-				Sample.Z *= -1.0f;
-				SampleDirections.Add(Sample);
-			}
+			FVector3f Sample = OtherHemisphereSamples[i];
+			Sample.Z *= -1.0f;
+			SampleDirections.Add(Sample);
 		}
+	}
 
-		static const auto CVar = IConsoleManager::Get().FindTConsoleVariableDataInt(TEXT("r.DistanceFields.MaxPerMeshResolution"));
-		const int32 PerMeshMax = CVar->GetValueOnAnyThread();
-
-		// Meshes with explicit artist-specified scale can go higher
-		const int32 MaxNumBlocksOneDim = FMath::Min<int32>(FMath::DivideAndRoundNearest(DistanceFieldResolutionScale <= 1 ? PerMeshMax / 2 : PerMeshMax, DistanceField::UniqueDataBrickSize), DistanceField::MaxIndirectionDimension - 1);
-
-		static const auto CVarDensity = IConsoleManager::Get().FindTConsoleVariableDataFloat(TEXT("r.DistanceFields.DefaultVoxelDensity"));
-		const float VoxelDensity = CVarDensity->GetValueOnAnyThread();
+	static const auto CVar = IConsoleManager::Get().FindTConsoleVariableDataInt(TEXT("r.DistanceFields.MaxPerMeshResolution"));
+	const int32 PerMeshMax = CVar->GetValueOnAnyThread();
 
-		const float NumVoxelsPerLocalSpaceUnit = VoxelDensity * DistanceFieldResolutionScale;
-		FBox3f LocalSpaceMeshBounds(Bounds.GetBox());
+	// Meshes with explicit artist-specified scale can go higher
+	const int32 MaxNumBlocksOneDim = FMath::Min<int32>(FMath::DivideAndRoundNearest(DistanceFieldResolutionScale <= 1 ? PerMeshMax / 2 : PerMeshMax, DistanceField::UniqueDataBrickSize), DistanceField::MaxIndirectionDimension - 1);
 
-		// Make sure the mesh bounding box has positive extents to handle planes
-		{
-			FVector3f MeshBoundsCenter = LocalSpaceMeshBounds.GetCenter();
-			FVector3f MeshBoundsExtent = FVector3f::Max(LocalSpaceMeshBounds.GetExtent(), FVector3f::OneVector);
-			LocalSpaceMeshBounds.Min = MeshBoundsCenter - MeshBoundsExtent;
-			LocalSpaceMeshBounds.Max = MeshBoundsCenter + MeshBoundsExtent;
-		}
+	static const auto CVarDensity = IConsoleManager::Get().FindTConsoleVariableDataFloat(TEXT("r.DistanceFields.DefaultVoxelDensity"));
+	const float VoxelDensity = CVarDensity->GetValueOnAnyThread();
 
-		// We sample on voxel corners and use central differencing for gradients, so a box mesh using two-sided materials whose vertices lie on LocalSpaceMeshBounds produces a zero gradient on intersection
-		// Expand the mesh bounds by a fraction of a voxel to allow room for a pullback on the hit location for computing the gradient.
-		// Only expand for two sided meshes as this adds significant Mesh SDF tracing cost
-		if (EmbreeScene.bMostlyTwoSided)
-		{
-			const FVector3f DesiredDimensions = LocalSpaceMeshBounds.GetSize() * FVector3f(NumVoxelsPerLocalSpaceUnit / (float)DistanceField::UniqueDataBrickSize);
-			const FInt32Vector Mip0IndirectionDimensions = FInt32Vector(
-				FMath::Clamp(FMath::RoundToInt(DesiredDimensions.X), 1, MaxNumBlocksOneDim),
-				FMath::Clamp(FMath::RoundToInt(DesiredDimensions.Y), 1, MaxNumBlocksOneDim),
-				FMath::Clamp(FMath::RoundToInt(DesiredDimensions.Z), 1, MaxNumBlocksOneDim));
-
-			const float CentralDifferencingExpandInVoxels = .25f;
-			const FVector3f TexelObjectSpaceSize = LocalSpaceMeshBounds.GetSize() / FVector3f(Mip0IndirectionDimensions * DistanceField::UniqueDataBrickSize - FInt32Vector(2 * CentralDifferencingExpandInVoxels));
-			LocalSpaceMeshBounds = LocalSpaceMeshBounds.ExpandBy(TexelObjectSpaceSize);
-		}
+	const float NumVoxelsPerLocalSpaceUnit = VoxelDensity * DistanceFieldResolutionScale;
+	FBox3f LocalSpaceMeshBounds(Bounds.GetBox());
 
-		// The tracing shader uses a Volume space that is normalized by the maximum extent, to keep Volume space within [-1, 1], we must match that behavior when encoding
-		const float LocalToVolumeScale = 1.0f / LocalSpaceMeshBounds.GetExtent().GetMax();
+	// Make sure the mesh bounding box has positive extents to handle planes
+	{
+		FVector3f MeshBoundsCenter = LocalSpaceMeshBounds.GetCenter();
+		FVector3f MeshBoundsExtent = FVector3f::Max(LocalSpaceMeshBounds.GetExtent(), FVector3f::OneVector);
+		LocalSpaceMeshBounds.Min = MeshBoundsCenter - MeshBoundsExtent;
+		LocalSpaceMeshBounds.Max = MeshBoundsCenter + MeshBoundsExtent;
+	}
 
-		const FVector3f DesiredDimensions = FVector3f(LocalSpaceMeshBounds.GetSize() * NumVoxelsPerLocalSpaceUnit / (float)DistanceField::UniqueDataBrickSize);
+	// We sample on voxel corners and use central differencing for gradients, so a box mesh using two-sided materials whose vertices lie on LocalSpaceMeshBounds produces a zero gradient on intersection
+	// Expand the mesh bounds by a fraction of a voxel to allow room for a pullback on the hit location for computing the gradient.
+	// Only expand for two sided meshes as this adds significant Mesh SDF tracing cost
+	if (EmbreeScene.bMostlyTwoSided)
+	{
+		const FVector3f DesiredDimensions = LocalSpaceMeshBounds.GetSize() * FVector3f(NumVoxelsPerLocalSpaceUnit / (float)DistanceField::UniqueDataBrickSize);
 		const FInt32Vector Mip0IndirectionDimensions = FInt32Vector(
 			FMath::Clamp(FMath::RoundToInt(DesiredDimensions.X), 1, MaxNumBlocksOneDim),
 			FMath::Clamp(FMath::RoundToInt(DesiredDimensions.Y), 1, MaxNumBlocksOneDim),
 			FMath::Clamp(FMath::RoundToInt(DesiredDimensions.Z), 1, MaxNumBlocksOneDim));
 
-		TArray<uint8> StreamableMipData;
+		const float CentralDifferencingExpandInVoxels = .25f;
+		const FVector3f TexelObjectSpaceSize = LocalSpaceMeshBounds.GetSize() / FVector3f(Mip0IndirectionDimensions * DistanceField::UniqueDataBrickSize - FInt32Vector(2 * CentralDifferencingExpandInVoxels));
+		LocalSpaceMeshBounds = LocalSpaceMeshBounds.ExpandBy(TexelObjectSpaceSize);
+	}
+
+	// The tracing shader uses a Volume space that is normalized by the maximum extent, to keep Volume space within [-1, 1], we must match that behavior when encoding
+	const float LocalToVolumeScale = 1.0f / LocalSpaceMeshBounds.GetExtent().GetMax();
 
-		for (int32 MipIndex = 0; MipIndex < DistanceField::NumMips; MipIndex++)
-		{
-			const FInt32Vector IndirectionDimensions = FInt32Vector(
-				FMath::DivideAndRoundUp(Mip0IndirectionDimensions.X, 1 << MipIndex),
-				FMath::DivideAndRoundUp(Mip0IndirectionDimensions.Y, 1 << MipIndex),
-				FMath::DivideAndRoundUp(Mip0IndirectionDimensions.Z, 1 << MipIndex));
+	const FVector3f DesiredDimensions = FVector3f(LocalSpaceMeshBounds.GetSize() * NumVoxelsPerLocalSpaceUnit / (float)DistanceField::UniqueDataBrickSize);
+	const FInt32Vector Mip0IndirectionDimensions = FInt32Vector(
+		FMath::Clamp(FMath::RoundToInt(DesiredDimensions.X), 1, MaxNumBlocksOneDim),
+		FMath::Clamp(FMath::RoundToInt(DesiredDimensions.Y), 1, MaxNumBlocksOneDim),
+		FMath::Clamp(FMath::RoundToInt(DesiredDimensions.Z), 1, MaxNumBlocksOneDim));
 
-			// Expand to guarantee one voxel border for gradient reconstruction using bilinear filtering
-			const FVector3f TexelObjectSpaceSize = LocalSpaceMeshBounds.GetSize() / FVector3f(IndirectionDimensions * DistanceField::UniqueDataBrickSize - FIntVector(2 * DistanceField::MeshDistanceFieldObjectBorder));
-			const FBox3f DistanceFieldVolumeBounds = LocalSpaceMeshBounds.ExpandBy(TexelObjectSpaceSize);
+	TArray<uint8> StreamableMipData;
 
-			const FVector3f IndirectionVoxelSize = DistanceFieldVolumeBounds.GetSize() / FVector3f(IndirectionDimensions);
-			const float IndirectionVoxelRadius = IndirectionVoxelSize.Size();
+	for (int32 MipIndex = 0; MipIndex < DistanceField::NumMips; MipIndex++)
+	{
+		const FInt32Vector IndirectionDimensions = FInt32Vector(
+			FMath::DivideAndRoundUp(Mip0IndirectionDimensions.X, 1 << MipIndex),
+			FMath::DivideAndRoundUp(Mip0IndirectionDimensions.Y, 1 << MipIndex),
+			FMath::DivideAndRoundUp(Mip0IndirectionDimensions.Z, 1 << MipIndex));
+
+		// Expand to guarantee one voxel border for gradient reconstruction using bilinear filtering
+		const FVector3f TexelObjectSpaceSize = LocalSpaceMeshBounds.GetSize() / FVector3f(IndirectionDimensions * DistanceField::UniqueDataBrickSize - FIntVector(2 * DistanceField::MeshDistanceFieldObjectBorder));
+		const FBox3f DistanceFieldVolumeBounds = LocalSpaceMeshBounds.ExpandBy(TexelObjectSpaceSize);
 
-			const FVector3f VolumeSpaceDistanceFieldVoxelSize = IndirectionVoxelSize * LocalToVolumeScale / FVector3f(DistanceField::UniqueDataBrickSize);
-			const float MaxDistanceForEncoding = VolumeSpaceDistanceFieldVoxelSize.Size() * DistanceField::BandSizeInVoxels;
-			const float LocalSpaceTraceDistance = MaxDistanceForEncoding / LocalToVolumeScale;
-			const FVector2f DistanceFieldToVolumeScaleBias(2.0f * MaxDistanceForEncoding, -MaxDistanceForEncoding);
+		const FVector3f IndirectionVoxelSize = DistanceFieldVolumeBounds.GetSize() / FVector3f(IndirectionDimensions);
+		const float IndirectionVoxelRadius = IndirectionVoxelSize.Size();
 
-			TArray<FSparseMeshDistanceFieldAsyncTask> AsyncTasks;
-			AsyncTasks.Reserve(IndirectionDimensions.X * IndirectionDimensions.Y * IndirectionDimensions.Z / 8);
+		const FVector3f VolumeSpaceDistanceFieldVoxelSize = IndirectionVoxelSize * LocalToVolumeScale / FVector3f(DistanceField::UniqueDataBrickSize);
+		const float MaxDistanceForEncoding = VolumeSpaceDistanceFieldVoxelSize.Size() * DistanceField::BandSizeInVoxels;
+		const float LocalSpaceTraceDistance = MaxDistanceForEncoding / LocalToVolumeScale;
+		const FVector2f DistanceFieldToVolumeScaleBias(2.0f * MaxDistanceForEncoding, -MaxDistanceForEncoding);
 
-			for (int32 ZIndex = 0; ZIndex < IndirectionDimensions.Z; ZIndex++)
+		TArray<FSparseMeshDistanceFieldAsyncTask> AsyncTasks;
+		AsyncTasks.Reserve(IndirectionDimensions.X * IndirectionDimensions.Y * IndirectionDimensions.Z / 8);
+
+		for (int32 ZIndex = 0; ZIndex < IndirectionDimensions.Z; ZIndex++)
+		{
+			for (int32 YIndex = 0; YIndex < IndirectionDimensions.Y; YIndex++)
 			{
-				for (int32 YIndex = 0; YIndex < IndirectionDimensions.Y; YIndex++)
+				for (int32 XIndex = 0; XIndex < IndirectionDimensions.X; XIndex++)
 				{
-					for (int32 XIndex = 0; XIndex < IndirectionDimensions.X; XIndex++)
-					{
-						AsyncTasks.Emplace(
-							EmbreeScene,
-							&SampleDirections,
-							LocalSpaceTraceDistance,
-							DistanceFieldVolumeBounds,
-							LocalToVolumeScale,
-							DistanceFieldToVolumeScaleBias,
-							FInt32Vector(XIndex, YIndex, ZIndex),
-							IndirectionDimensions,
-							bUsePointQuery);
-					}
+					AsyncTasks.Emplace(
+						EmbreeScene,
+						&SampleDirections,
+						LocalSpaceTraceDistance,
+						DistanceFieldVolumeBounds,
+						LocalToVolumeScale,
+						DistanceFieldToVolumeScaleBias,
+						FInt32Vector(XIndex, YIndex, ZIndex),
+						IndirectionDimensions,
+						bUsePointQuery);
 				}
 			}
+		}
 
-			static bool bMultiThreaded = true;
+		static bool bMultiThreaded = true;
 
-			if (bMultiThreaded)
-			{
-				EParallelForFlags Flags = EParallelForFlags::BackgroundPriority | EParallelForFlags::Unbalanced;
-				COOK_STAT(auto AsyncTimer = DistanceFieldUtilityCookStats::UsageStats.TimeAsyncWait());
-				COOK_STAT(AsyncTimer.TrackCyclesOnly());
+		if (bMultiThreaded)
+		{
+			EParallelForFlags Flags = EParallelForFlags::BackgroundPriority | EParallelForFlags::Unbalanced;
+			COOK_STAT(auto AsyncTimer = DistanceFieldUtilityCookStats::UsageStats.TimeAsyncWait());
+			COOK_STAT(AsyncTimer.TrackCyclesOnly());
 
-				ParallelForTemplate(
-					TEXT("GenerateSignedDistanceFieldVolumeData.PF"),
-					AsyncTasks.Num(),1, [&AsyncTasks](int32 TaskIndex)
+			ParallelForTemplate(
+				TEXT("GenerateSignedDistanceFieldVolumeData.PF"),
+				AsyncTasks.Num(), 1, [&AsyncTasks](int32 TaskIndex)
 				{
 					AsyncTasks[TaskIndex].DoWork();
 				}, Flags);
-			}
-			else
+		}
+		else
+		{
+			for (FSparseMeshDistanceFieldAsyncTask& AsyncTask : AsyncTasks)
 			{
-				for (FSparseMeshDistanceFieldAsyncTask& AsyncTask : AsyncTasks)
-				{
-					AsyncTask.DoWork();
-				}
+				AsyncTask.DoWork();
 			}
+		}
 
-			FSparseDistanceFieldMip& OutMip = OutData.Mips[MipIndex];
-			TArray<uint32> IndirectionTable;
-			IndirectionTable.Empty(IndirectionDimensions.X * IndirectionDimensions.Y * IndirectionDimensions.Z);
-			IndirectionTable.AddUninitialized(IndirectionDimensions.X * IndirectionDimensions.Y * IndirectionDimensions.Z);
+		FSparseDistanceFieldMip& OutMip = OutData.Mips[MipIndex];
+		TArray<uint32> IndirectionTable;
+		IndirectionTable.Empty(IndirectionDimensions.X * IndirectionDimensions.Y * IndirectionDimensions.Z);
+		IndirectionTable.AddUninitialized(IndirectionDimensions.X * IndirectionDimensions.Y * IndirectionDimensions.Z);
 
-			for (int32 i = 0; i < IndirectionTable.Num(); i++)
-			{
-				IndirectionTable[i] = DistanceField::InvalidBrickIndex;
-			} 
+		for (int32 i = 0; i < IndirectionTable.Num(); i++)
+		{
+			IndirectionTable[i] = DistanceField::InvalidBrickIndex;
+		}
 
-			TArray<FSparseMeshDistanceFieldAsyncTask*> ValidBricks;
-			ValidBricks.Empty(AsyncTasks.Num());
+		TArray<FSparseMeshDistanceFieldAsyncTask*> ValidBricks;
+		ValidBricks.Empty(AsyncTasks.Num());
 
-			for (int32 TaskIndex = 0; TaskIndex < AsyncTasks.Num(); TaskIndex++)
+		for (int32 TaskIndex = 0; TaskIndex < AsyncTasks.Num(); TaskIndex++)
+		{
+			if (AsyncTasks[TaskIndex].BrickMinDistance < MAX_uint8 && AsyncTasks[TaskIndex].BrickMaxDistance > MIN_uint8)
 			{
-				if (AsyncTasks[TaskIndex].BrickMinDistance < MAX_uint8 && AsyncTasks[TaskIndex].BrickMaxDistance > MIN_uint8)
-				{
-					ValidBricks.Add(&AsyncTasks[TaskIndex]);
-				}
+				ValidBricks.Add(&AsyncTasks[TaskIndex]);
 			}
+		}
 
-			const uint32 NumBricks = ValidBricks.Num();
+		const uint32 NumBricks = ValidBricks.Num();
 
-			const uint32 BrickSizeBytes = DistanceField::BrickSize * DistanceField::BrickSize * DistanceField::BrickSize * GPixelFormats[DistanceField::DistanceFieldFormat].BlockBytes;
+		const uint32 BrickSizeBytes = DistanceField::BrickSize * DistanceField::BrickSize * DistanceField::BrickSize * GPixelFormats[DistanceField::DistanceFieldFormat].BlockBytes;
 
-			TArray<uint8> DistanceFieldBrickData;
-			DistanceFieldBrickData.Empty(BrickSizeBytes * NumBricks);
-			DistanceFieldBrickData.AddUninitialized(BrickSizeBytes * NumBricks);
+		TArray<uint8> DistanceFieldBrickData;
+		DistanceFieldBrickData.Empty(BrickSizeBytes * NumBricks);
+		DistanceFieldBrickData.AddUninitialized(BrickSizeBytes * NumBricks);
 
-			for (int32 BrickIndex = 0; BrickIndex < ValidBricks.Num(); BrickIndex++)
-			{
-				const FSparseMeshDistanceFieldAsyncTask& Brick = *ValidBricks[BrickIndex];
-				const int32 IndirectionIndex = ComputeLinearVoxelIndex(Brick.BrickCoordinate, IndirectionDimensions);
-				IndirectionTable[IndirectionIndex] = BrickIndex;
+		for (int32 BrickIndex = 0; BrickIndex < ValidBricks.Num(); BrickIndex++)
+		{
+			const FSparseMeshDistanceFieldAsyncTask& Brick = *ValidBricks[BrickIndex];
+			const int32 IndirectionIndex = ComputeLinearVoxelIndex(Brick.BrickCoordinate, IndirectionDimensions);
+			IndirectionTable[IndirectionIndex] = BrickIndex;
 
-				check(BrickSizeBytes == Brick.DistanceFieldVolume.Num() * Brick.DistanceFieldVolume.GetTypeSize());
-				FPlatformMemory::Memcpy(&DistanceFieldBrickData[BrickIndex * BrickSizeBytes], Brick.DistanceFieldVolume.GetData(), Brick.DistanceFieldVolume.Num() * Brick.DistanceFieldVolume.GetTypeSize());
-			}
+			check(BrickSizeBytes == Brick.DistanceFieldVolume.Num() * Brick.DistanceFieldVolume.GetTypeSize());
+			FPlatformMemory::Memcpy(&DistanceFieldBrickData[BrickIndex * BrickSizeBytes], Brick.DistanceFieldVolume.GetData(), Brick.DistanceFieldVolume.Num() * Brick.DistanceFieldVolume.GetTypeSize());
+		}
 
-			const int32 IndirectionTableBytes = IndirectionTable.Num() * IndirectionTable.GetTypeSize();
-			const int32 MipDataBytes = IndirectionTableBytes + DistanceFieldBrickData.Num();
+		const int32 IndirectionTableBytes = IndirectionTable.Num() * IndirectionTable.GetTypeSize();
+		const int32 MipDataBytes = IndirectionTableBytes + DistanceFieldBrickData.Num();
 
-			if (MipIndex == DistanceField::NumMips - 1)
-			{
-				OutData.AlwaysLoadedMip.Empty(MipDataBytes);
-				OutData.AlwaysLoadedMip.AddUninitialized(MipDataBytes);
+		if (MipIndex == DistanceField::NumMips - 1)
+		{
+			OutData.AlwaysLoadedMip.Empty(MipDataBytes);
+			OutData.AlwaysLoadedMip.AddUninitialized(MipDataBytes);
 
-				FPlatformMemory::Memcpy(&OutData.AlwaysLoadedMip[0], IndirectionTable.GetData(), IndirectionTableBytes);
+			FPlatformMemory::Memcpy(&OutData.AlwaysLoadedMip[0], IndirectionTable.GetData(), IndirectionTableBytes);
 
-				if (DistanceFieldBrickData.Num() > 0)
-				{
-					FPlatformMemory::Memcpy(&OutData.AlwaysLoadedMip[IndirectionTableBytes], DistanceFieldBrickData.GetData(), DistanceFieldBrickData.Num());
-				}
-			}
-			else
+			if (DistanceFieldBrickData.Num() > 0)
 			{
-				OutMip.BulkOffset = StreamableMipData.Num();
-				StreamableMipData.AddUninitialized(MipDataBytes);
-				OutMip.BulkSize = StreamableMipData.Num() - OutMip.BulkOffset;
-				checkf(OutMip.BulkSize > 0, TEXT("BulkSize was 0 for %s with %ux%ux%u indirection"), *MeshName, IndirectionDimensions.X, IndirectionDimensions.Y, IndirectionDimensions.Z);
+				FPlatformMemory::Memcpy(&OutData.AlwaysLoadedMip[IndirectionTableBytes], DistanceFieldBrickData.GetData(), DistanceFieldBrickData.Num());
+			}
+		}
+		else
+		{
+			OutMip.BulkOffset = StreamableMipData.Num();
+			StreamableMipData.AddUninitialized(MipDataBytes);
+			OutMip.BulkSize = StreamableMipData.Num() - OutMip.BulkOffset;
+			checkf(OutMip.BulkSize > 0, TEXT("BulkSize was 0 for %s with %ux%ux%u indirection"), *MeshName, IndirectionDimensions.X, IndirectionDimensions.Y, IndirectionDimensions.Z);
 
-				FPlatformMemory::Memcpy(&StreamableMipData[OutMip.BulkOffset], IndirectionTable.GetData(), IndirectionTableBytes);
+			FPlatformMemory::Memcpy(&StreamableMipData[OutMip.BulkOffset], IndirectionTable.GetData(), IndirectionTableBytes);
 
-				if (DistanceFieldBrickData.Num() > 0)
-				{
-					FPlatformMemory::Memcpy(&StreamableMipData[OutMip.BulkOffset + IndirectionTableBytes], DistanceFieldBrickData.GetData(), DistanceFieldBrickData.Num());
-				}
+			if (DistanceFieldBrickData.Num() > 0)
+			{
+				FPlatformMemory::Memcpy(&StreamableMipData[OutMip.BulkOffset + IndirectionTableBytes], DistanceFieldBrickData.GetData(), DistanceFieldBrickData.Num());
 			}
-	
-			OutMip.IndirectionDimensions = IndirectionDimensions;
-			OutMip.DistanceFieldToVolumeScaleBias = DistanceFieldToVolumeScaleBias;
-			OutMip.NumDistanceFieldBricks = NumBricks;
-
-			// Account for the border voxels we added
-			const FVector3f VirtualUVMin = FVector3f(DistanceField::MeshDistanceFieldObjectBorder) / FVector3f(IndirectionDimensions * DistanceField::UniqueDataBrickSize);
-			const FVector3f VirtualUVSize = FVector3f(IndirectionDimensions * DistanceField::UniqueDataBrickSize - FInt32Vector(2 * DistanceField::MeshDistanceFieldObjectBorder)) / FVector3f(IndirectionDimensions * DistanceField::UniqueDataBrickSize);
-		
-			const FVector3f VolumePositionExtent = LocalSpaceMeshBounds.GetExtent() * LocalToVolumeScale;
-
-			// [-VolumePositionExtent, VolumePositionExtent] -> [VirtualUVMin, VirtualUVMin + VirtualUVSize]
-			OutMip.VolumeToVirtualUVScale = VirtualUVSize / (2 * VolumePositionExtent);
-			OutMip.VolumeToVirtualUVAdd = VolumePositionExtent * OutMip.VolumeToVirtualUVScale + VirtualUVMin;
 		}
 
-		MeshRepresentation::DeleteEmbreeScene(EmbreeScene);
+		OutMip.IndirectionDimensions = IndirectionDimensions;
+		OutMip.DistanceFieldToVolumeScaleBias = DistanceFieldToVolumeScaleBias;
+		OutMip.NumDistanceFieldBricks = NumBricks;
+
+		// Account for the border voxels we added
+		const FVector3f VirtualUVMin = FVector3f(DistanceField::MeshDistanceFieldObjectBorder) / FVector3f(IndirectionDimensions * DistanceField::UniqueDataBrickSize);
+		const FVector3f VirtualUVSize = FVector3f(IndirectionDimensions * DistanceField::UniqueDataBrickSize - FInt32Vector(2 * DistanceField::MeshDistanceFieldObjectBorder)) / FVector3f(IndirectionDimensions * DistanceField::UniqueDataBrickSize);
 
-		OutData.bMostlyTwoSided = EmbreeScene.bMostlyTwoSided;
-		OutData.LocalSpaceMeshBounds = LocalSpaceMeshBounds;
+		const FVector3f VolumePositionExtent = LocalSpaceMeshBounds.GetExtent() * LocalToVolumeScale;
 
-		OutData.StreamableMips.Lock(LOCK_READ_WRITE);
-		uint8* Ptr = (uint8*)OutData.StreamableMips.Realloc(StreamableMipData.Num());
-		FMemory::Memcpy(Ptr, StreamableMipData.GetData(), StreamableMipData.Num());
-		OutData.StreamableMips.Unlock();
-		OutData.StreamableMips.SetBulkDataFlags(BULKDATA_Force_NOT_InlinePayload);
+		// [-VolumePositionExtent, VolumePositionExtent] -> [VirtualUVMin, VirtualUVMin + VirtualUVSize]
+		OutMip.VolumeToVirtualUVScale = VirtualUVSize / (2 * VolumePositionExtent);
+		OutMip.VolumeToVirtualUVAdd = VolumePositionExtent * OutMip.VolumeToVirtualUVScale + VirtualUVMin;
+	}
+
+	OutData.bMostlyTwoSided = EmbreeScene.bMostlyTwoSided;
+	OutData.LocalSpaceMeshBounds = LocalSpaceMeshBounds;
+
+	OutData.StreamableMips.Lock(LOCK_READ_WRITE);
+	uint8* Ptr = (uint8*)OutData.StreamableMips.Realloc(StreamableMipData.Num());
+	FMemory::Memcpy(Ptr, StreamableMipData.GetData(), StreamableMipData.Num());
+	OutData.StreamableMips.Unlock();
+	OutData.StreamableMips.SetBulkDataFlags(BULKDATA_Force_NOT_InlinePayload);
+}
+
+void FMeshUtilities::GenerateSignedDistanceFieldVolumeData(
+	FString MeshName,
+	const FSourceMeshDataForDerivedDataTask& SourceMeshData,
+	const FStaticMeshLODResources& LODModel,
+	class FQueuedThreadPool& ThreadPool,
+	const TArray<FSignedDistanceFieldBuildSectionData>& SectionData,
+	const FBoxSphereBounds3f& Bounds,
+	float DistanceFieldResolutionScale,
+	bool bGenerateAsIfTwoSided,
+	FDistanceFieldVolumeData& OutData)
+{
+	TRACE_CPUPROFILER_EVENT_SCOPE(GenerateSignedDistanceFieldVolumeData);
+	COOK_STAT(auto Timer = DistanceFieldUtilityCookStats::UsageStats.TimeSyncWork());
+	COOK_STAT(Timer.TrackCyclesOnly());
+
+	if (DistanceFieldResolutionScale > 0)
+	{
+		const double StartTime = FPlatformTime::Seconds();
+		const bool bIncludeTranslucentTriangles = false;
+
+		FEmbreeScene EmbreeScene;
+		MeshRepresentation::SetupEmbreeScene(MeshName,
+			SourceMeshData,
+			LODModel,
+			SectionData,
+			bGenerateAsIfTwoSided,
+			bIncludeTranslucentTriangles,
+			EmbreeScene);
+
+		// If Embree setup fails, there will be no scene to operate on. Early out.
+		if (!EmbreeScene.EmbreeScene)
+		{
+			return;
+		}
+
+		BuildSignedDistanceField(MeshName, EmbreeScene, Bounds, DistanceFieldResolutionScale, OutData);
+
+		MeshRepresentation::DeleteEmbreeScene(EmbreeScene);
 
 		const float BuildTime = (float)(FPlatformTime::Seconds() - StartTime);
 		 
 		if (BuildTime > 1.0f)
 		{
+			const FInt32Vector Mip0IndirectionDimensions = OutData.Mips[0].IndirectionDimensions;
+
 			UE_LOG(LogMeshUtilities, Log, TEXT("Finished distance field build in %.1fs - %ux%ux%u sparse distance field, %.1fMb total, %.1fMb always loaded, %u%% occupied, %u triangles, %s"),
 				BuildTime,
 				Mip0IndirectionDimensions.X * DistanceField::UniqueDataBrickSize,
