Improved MeshRepresentationCommon to support instancing in Embree scene.

diff --git a/Engine/Source/Developer/MeshUtilities/Private/MeshCardRepresentationUtilities.cpp b/Engine/Source/Developer/MeshUtilities/Private/MeshCardRepresentationUtilities.cpp
--- a/Engine/Source/Developer/MeshUtilities/Private/MeshCardRepresentationUtilities.cpp
+++ b/Engine/Source/Developer/MeshUtilities/Private/MeshCardRepresentationUtilities.cpp
@@ -289,7 +289,7 @@ FSurfelVisibility ComputeSurfelVisibility(
 
 		FEmbreeIntersectionContext EmbreeContext;
 		rtcInitIntersectContext(&EmbreeContext);
-		rtcIntersect1(Context.EmbreeScene.EmbreeScene, &EmbreeContext, &EmbreeRay);
+		rtcIntersect1(Context.EmbreeScene.Scene, &EmbreeContext, &EmbreeRay);
 
 		if (EmbreeRay.hit.geomID != RTC_INVALID_GEOMETRY_ID && EmbreeRay.hit.primID != RTC_INVALID_GEOMETRY_ID)
 		{
@@ -395,7 +395,7 @@ void GenerateSurfelsForDirection(
 					FEmbreeIntersectionContext EmbreeContext;
 					rtcInitIntersectContext(&EmbreeContext);
 					EmbreeContext.SkipPrimId = SkipPrimId;
-					rtcIntersect1(Context.EmbreeScene.EmbreeScene, &EmbreeContext, &EmbreeRay);
+					rtcIntersect1(Context.EmbreeScene.Scene, &EmbreeContext, &EmbreeRay);
 
 					if (EmbreeRay.hit.geomID != RTC_INVALID_GEOMETRY_ID && EmbreeRay.hit.primID != RTC_INVALID_GEOMETRY_ID)
 					{
@@ -653,11 +653,7 @@ void InitSurfelScene(
 {
 	TRACE_CPUPROFILER_EVENT_SCOPE(GenerateSurfels);
 
-	const uint32 NumSourceVertices = Context.EmbreeScene.Geometry.VertexArray.Num();
-	const uint32 NumSourceIndices = Context.EmbreeScene.Geometry.IndexArray.Num();
-	const int32 NumSourceTriangles = NumSourceIndices / 3;
-
-	if (NumSourceTriangles == 0)
+	if (Context.EmbreeScene.NumTrianglesTotal == 0)
 	{
 		return;
 	}
@@ -1073,19 +1069,19 @@ bool FMeshUtilities::GenerateCardRepresentationData(
 		const bool bIncludeTranslucentTriangles = true; 
 
 		FEmbreeScene EmbreeScene;
-		MeshRepresentation::SetupEmbreeScene(MeshName,
-			SourceMeshData,
-			LODModel,
-			SectionData,
-			bGenerateAsIfTwoSided,
-			bIncludeTranslucentTriangles,
-			EmbreeScene);
-
-		if (!EmbreeScene.EmbreeScene)
+		MeshRepresentation::SetupEmbreeScene(MeshName, bGenerateAsIfTwoSided, EmbreeScene);
+
+		if (!EmbreeScene.Scene)
 		{
 			return false;
 		}
 
+		const FEmbreeGeometryAsset* GeometryAsset = EmbreeScene.AddGeometryAsset(SourceMeshData, &LODModel, SectionData, bIncludeTranslucentTriangles, false);
+
+		const FEmbreeGeometry* Geometry = EmbreeScene.AddGeometry(GeometryAsset);
+
+		EmbreeScene.Commit();
+
 		FGenerateCardMeshContext Context(MeshName, EmbreeScene, OutData);
 
 		// Note: must operate on the SDF bounds when available, because SDF generation can expand the mesh's bounds
@@ -1100,7 +1096,7 @@ bool FMeshUtilities::GenerateCardRepresentationData(
 			UE_LOG(LogMeshUtilities, Log, TEXT("Finished mesh card build in %.1fs %s tris:%d surfels:%d"),
 				TimeElapsed,
 				*MeshName,
-				EmbreeScene.NumIndices / 3,
+				EmbreeScene.NumTrianglesTotal,
 				OutData.MeshCardsBuildData.DebugData.NumSurfels);
 		}
 	}
diff --git a/Engine/Source/Developer/MeshUtilities/Private/MeshDistanceFieldUtilities.cpp b/Engine/Source/Developer/MeshUtilities/Private/MeshDistanceFieldUtilities.cpp
--- a/Engine/Source/Developer/MeshUtilities/Private/MeshDistanceFieldUtilities.cpp
+++ b/Engine/Source/Developer/MeshUtilities/Private/MeshDistanceFieldUtilities.cpp
@@ -29,8 +29,7 @@ namespace DistanceFieldUtilityCookStats
 class FEmbreePointQueryContext : public RTCPointQueryContext
 {
 public:
-	RTCGeometry MeshGeometry;
-	int32 NumTriangles;
+	const FEmbreeScene* Scene = nullptr;
 };
 
 bool EmbreePointQueryFunction(RTCPointQueryFunctionArguments* args)
@@ -40,21 +39,43 @@ bool EmbreePointQueryFunction(RTCPointQueryFunctionArguments* args)
 	check(args->userPtr);
 	float& ClosestDistanceSq = *(float*)(args->userPtr);
 
+	int32 GeometryIndex = args->geomID;
+
+	if (Context->instID[0] != RTC_INVALID_GEOMETRY_ID)
+	{
+		// when testing against a geometry instance use instID to index into Scene->Geometries
+		GeometryIndex = Context->instID[0];
+	}
+
+	const FEmbreeGeometryAsset* GeometryAsset = Context->Scene->Geometries[GeometryIndex].Asset;
+	const int32 NumTriangles = GeometryAsset->NumTriangles;
+
 	const int32 TriangleIndex = args->primID;
-	check(TriangleIndex < Context->NumTriangles);
+	check(TriangleIndex < NumTriangles);
 
-	const FVector3f* VertexBuffer = (const FVector3f*)rtcGetGeometryBufferData(Context->MeshGeometry, RTC_BUFFER_TYPE_VERTEX, 0);
-	const uint32* IndexBuffer = (const uint32*)rtcGetGeometryBufferData(Context->MeshGeometry, RTC_BUFFER_TYPE_INDEX, 0);
+	const FVector3f* VertexBuffer = (const FVector3f*)GeometryAsset->VertexArray.GetData();
+	const uint32* IndexBuffer = (const uint32*)GeometryAsset->IndexArray.GetData();
 
 	const uint32 I0 = IndexBuffer[TriangleIndex * 3 + 0];
 	const uint32 I1 = IndexBuffer[TriangleIndex * 3 + 1];
 	const uint32 I2 = IndexBuffer[TriangleIndex * 3 + 2];
 
-	const FVector3f V0 = VertexBuffer[I0];
-	const FVector3f V1 = VertexBuffer[I1];
-	const FVector3f V2 = VertexBuffer[I2];
+	FVector3f V0 = VertexBuffer[I0];
+	FVector3f V1 = VertexBuffer[I1];
+	FVector3f V2 = VertexBuffer[I2];
+
+	if (Context->instID[0] != RTC_INVALID_GEOMETRY_ID)
+	{
+		// when testing against a geometry instance need to transform vertices to world space
+		FMatrix44f* InstToWorld = (FMatrix44f*)Context->inst2world[0];
+
+		V0 = InstToWorld->TransformPosition(V0);
+		V1 = InstToWorld->TransformPosition(V1);
+		V2 = InstToWorld->TransformPosition(V2);
+	}
 
 	const FVector3f QueryPosition(args->query->x, args->query->y, args->query->z);
+
 	const FVector3f ClosestPoint = (FVector3f)FMath::ClosestPointOnTriangleToPoint((FVector)QueryPosition, (FVector)V0, (FVector)V1, (FVector)V2);
 	const float QueryDistanceSq = (ClosestPoint - QueryPosition).SizeSquared();
 
@@ -173,10 +194,9 @@ void FSparseMeshDistanceFieldAsyncTask::DoWork()
 
 					FEmbreePointQueryContext QueryContext;
 					rtcInitPointQueryContext(&QueryContext);
-					QueryContext.MeshGeometry = EmbreeScene.Geometry.InternalGeometry;
-					QueryContext.NumTriangles = EmbreeScene.Geometry.TriangleDescs.Num();
+					QueryContext.Scene = &EmbreeScene;
 					float ClosestUnsignedDistanceSq = (LocalSpaceTraceDistance * 2.0f) * (LocalSpaceTraceDistance * 2.0f);
-					rtcPointQuery(EmbreeScene.EmbreeScene, &PointQuery, &QueryContext, EmbreePointQueryFunction, &ClosestUnsignedDistanceSq);
+					rtcPointQuery(EmbreeScene.Scene, &PointQuery, &QueryContext, EmbreePointQueryFunction, &ClosestUnsignedDistanceSq);
 
 					const float ClosestDistance = FMath::Sqrt(ClosestUnsignedDistanceSq);
 					bTraceRays = ClosestDistance <= LocalSpaceTraceDistance;
@@ -213,7 +233,7 @@ void FSparseMeshDistanceFieldAsyncTask::DoWork()
 
 							FEmbreeIntersectionContext EmbreeContext;
 							rtcInitIntersectContext(&EmbreeContext);
-							rtcIntersect1(EmbreeScene.EmbreeScene, &EmbreeContext, &EmbreeRay);
+							rtcIntersect1(EmbreeScene.Scene, &EmbreeContext, &EmbreeRay);
 
 							if (EmbreeRay.hit.geomID != RTC_INVALID_GEOMETRY_ID && EmbreeRay.hit.primID != RTC_INVALID_GEOMETRY_ID)
 							{
@@ -514,20 +534,20 @@ void FMeshUtilities::GenerateSignedDistanceFieldVolumeData(
 		const bool bIncludeTranslucentTriangles = false;
 
 		FEmbreeScene EmbreeScene;
-		MeshRepresentation::SetupEmbreeScene(MeshName,
-			SourceMeshData,
-			LODModel,
-			SectionData,
-			bGenerateAsIfTwoSided,
-			bIncludeTranslucentTriangles,
-			EmbreeScene);
+		MeshRepresentation::SetupEmbreeScene(MeshName, bGenerateAsIfTwoSided, EmbreeScene);
 
 		// If Embree setup fails, there will be no scene to operate on. Early out.
-		if (!EmbreeScene.EmbreeScene)
+		if (!EmbreeScene.Scene)
 		{
 			return;
 		}
 
+		const FEmbreeGeometryAsset* GeometryAsset = EmbreeScene.AddGeometryAsset(SourceMeshData, &LODModel, SectionData, bIncludeTranslucentTriangles, false);
+
+		const FEmbreeGeometry* Geometry = EmbreeScene.AddGeometry(GeometryAsset);
+
+		EmbreeScene.Commit();
+
 		BuildSignedDistanceField(MeshName, EmbreeScene, Bounds, DistanceFieldResolutionScale, OutData);
 
 		MeshRepresentation::DeleteEmbreeScene(EmbreeScene);
@@ -546,7 +566,7 @@ void FMeshUtilities::GenerateSignedDistanceFieldVolumeData(
 				(OutData.GetResourceSizeBytes() + OutData.StreamableMips.GetBulkDataSize()) / 1024.0f / 1024.0f,
 				(OutData.AlwaysLoadedMip.GetAllocatedSize()) / 1024.0f / 1024.0f,
 				FMath::RoundToInt(100.0f * OutData.Mips[0].NumDistanceFieldBricks / (float)(Mip0IndirectionDimensions.X * Mip0IndirectionDimensions.Y * Mip0IndirectionDimensions.Z)),
-				EmbreeScene.NumIndices / 3,
+				EmbreeScene.NumTrianglesTotal,
 				*MeshName);
 		}
 	}
diff --git a/Engine/Source/Developer/MeshUtilities/Private/MeshRepresentationCommon.cpp b/Engine/Source/Developer/MeshUtilities/Private/MeshRepresentationCommon.cpp
--- a/Engine/Source/Developer/MeshUtilities/Private/MeshRepresentationCommon.cpp
+++ b/Engine/Source/Developer/MeshUtilities/Private/MeshRepresentationCommon.cpp
@@ -94,8 +94,8 @@ FMatrix44f MeshRepresentation::GetTangentBasisFrisvad(FVector3f TangentZ)
 #if USE_EMBREE
 void EmbreeFilterFunc(const struct RTCFilterFunctionNArguments* args)
 {
-	FEmbreeGeometry* EmbreeGeometry = (FEmbreeGeometry*)args->geometryUserPtr;
-	FEmbreeTriangleDesc Desc = EmbreeGeometry->TriangleDescs[RTCHitN_primID(args->hit, 1, 0)];
+	const FEmbreeGeometryAsset* GeometryAsset = (const FEmbreeGeometryAsset*)args->geometryUserPtr;
+	FEmbreeTriangleDesc Desc = GeometryAsset->TriangleDescs[RTCHitN_primID(args->hit, 1, 0)];
 
 	FEmbreeIntersectionContext& IntersectionContext = *static_cast<FEmbreeIntersectionContext*>(args->context);
 	IntersectionContext.ElementIndex = Desc.ElementIndex;
@@ -124,69 +124,85 @@ void EmbreeErrorFunc(void* userPtr, RTCError code, const char* str)
 }
 #endif
 
-void MeshRepresentation::SetupEmbreeScene(
-	FString MeshName,
-	const FSourceMeshDataForDerivedDataTask& SourceMeshData,
-	const FStaticMeshLODResources& LODModel,
-	const TArray<FSignedDistanceFieldBuildSectionData>& SectionData,
-	bool bGenerateAsIfTwoSided,
-	bool bIncludeTranslucentTriangles,
-	FEmbreeScene& EmbreeScene)
+void MeshRepresentation::SetupEmbreeScene(FString MeshName, bool bGenerateAsIfTwoSided, FEmbreeScene& OutEmbreeScene)
 {
-	const uint32 NumVertices = SourceMeshData.IsValid() ? SourceMeshData.GetNumVertices() : LODModel.VertexBuffers.PositionVertexBuffer.GetNumVertices();
-	const uint32 NumIndices = SourceMeshData.IsValid() ? SourceMeshData.GetNumIndices() : LODModel.IndexBuffer.GetNumIndices();
-	const int32 NumTriangles = NumIndices / 3;
-	EmbreeScene.NumIndices = NumTriangles;
-
-	const FStaticMeshSectionArray& Sections = SourceMeshData.IsValid() ? SourceMeshData.Sections : LODModel.Sections;
+	OutEmbreeScene.MeshName = MeshName;
+	OutEmbreeScene.bGenerateAsIfTwoSided = bGenerateAsIfTwoSided;
 
 #if USE_EMBREE
+	OutEmbreeScene.Device = rtcNewDevice(nullptr);
+	rtcSetDeviceErrorFunction(OutEmbreeScene.Device, EmbreeErrorFunc, nullptr);
+
+	RTCError ReturnErrorNewDevice = rtcGetDeviceError(OutEmbreeScene.Device);
+	if (ReturnErrorNewDevice == RTC_ERROR_OUT_OF_MEMORY)
+	{
+		UE_LOG(LogMeshUtilities, Warning, TEXT("Failed to created Embree device for %s (OUT_OF_MEMORY)."), *MeshName);
+		FPlatformMemory::OnOutOfMemory(0, 16);
+		return;
+	}
+	if (ReturnErrorNewDevice != RTC_ERROR_NONE)
 	{
-		EmbreeScene.EmbreeDevice = rtcNewDevice(nullptr);
-		rtcSetDeviceErrorFunction(EmbreeScene.EmbreeDevice, EmbreeErrorFunc, nullptr);
+		UE_LOG(LogMeshUtilities, Warning, TEXT("Failed to created Embree device for %s. Code: %d"), *MeshName, (int32)ReturnErrorNewDevice);
+		return;
+	}
 
-		RTCError ReturnErrorNewDevice = rtcGetDeviceError(EmbreeScene.EmbreeDevice);
-		if (ReturnErrorNewDevice == RTC_ERROR_OUT_OF_MEMORY)
-		{
-			UE_LOG(LogMeshUtilities, Warning, TEXT("GenerateSignedDistanceFieldVolumeData failed for %s. Embree rtcNewDevice failed to allocate memory."), *MeshName);
-			FPlatformMemory::OnOutOfMemory(0, 16);
-			return;
-		}
-		if (ReturnErrorNewDevice != RTC_ERROR_NONE)
-		{
-			UE_LOG(LogMeshUtilities, Warning, TEXT("GenerateSignedDistanceFieldVolumeData failed for %s. Embree rtcNewDevice failed. Code: %d"), *MeshName, (int32)ReturnErrorNewDevice);
-			return;
-		}
+	OutEmbreeScene.Scene = rtcNewScene(OutEmbreeScene.Device);
+	rtcSetSceneFlags(OutEmbreeScene.Scene, RTC_SCENE_FLAG_NONE);
 
-		EmbreeScene.EmbreeScene = rtcNewScene(EmbreeScene.EmbreeDevice);
-		rtcSetSceneFlags(EmbreeScene.EmbreeScene, RTC_SCENE_FLAG_NONE);
+	RTCError ReturnErrorNewScene = rtcGetDeviceError(OutEmbreeScene.Device);
+	if (ReturnErrorNewScene == RTC_ERROR_OUT_OF_MEMORY)
+	{
+		UE_LOG(LogMeshUtilities, Warning, TEXT("Failed to created Embree scene for %s (OUT_OF_MEMORY)."), *MeshName);
+		FPlatformMemory::OnOutOfMemory(0, 16);
+		return;
+	}
+	if (ReturnErrorNewScene != RTC_ERROR_NONE)
+	{
+		UE_LOG(LogMeshUtilities, Warning, TEXT("Failed to created Embree scene for %s. Code: %d"), *MeshName, (int32)ReturnErrorNewScene);
+		rtcReleaseDevice(OutEmbreeScene.Device);
+		return;
+	}
+#endif
+}
 
-		RTCError ReturnErrorNewScene = rtcGetDeviceError(EmbreeScene.EmbreeDevice);
-		if (ReturnErrorNewScene == RTC_ERROR_OUT_OF_MEMORY)
-		{
-			UE_LOG(LogMeshUtilities, Warning, TEXT("GenerateSignedDistanceFieldVolumeData failed for %s. Embree rtcNewScene failed to allocate memory."), *MeshName);
-			FPlatformMemory::OnOutOfMemory(0, 16);
-			return;
-		}
-		if (ReturnErrorNewScene != RTC_ERROR_NONE)
+void MeshRepresentation::DeleteEmbreeScene(FEmbreeScene& EmbreeScene)
+{
+#if USE_EMBREE
+	for (FEmbreeGeometryAsset& Asset : EmbreeScene.GeometryAssets)
+	{
+		if (Asset.Scene != nullptr)
 		{
-			UE_LOG(LogMeshUtilities, Warning, TEXT("GenerateSignedDistanceFieldVolumeData failed for %s. Embree rtcNewScene failed. Code: %d"), *MeshName, (int32)ReturnErrorNewScene);
-			rtcReleaseDevice(EmbreeScene.EmbreeDevice);
-			return;
+			rtcReleaseScene(Asset.Scene);
 		}
 	}
+
+	rtcReleaseScene(EmbreeScene.Scene);
+	rtcReleaseDevice(EmbreeScene.Device);
 #endif
 
-	/*
-	if (LODModel.Sections.Num() > SectionData.Num())
-	{
-		UE_LOG(LogMeshUtilities, Warning, TEXT("Unexpected number of mesh sections when setting up Embree Scene for %s."), *MeshName);
-	}
-	*/
+	EmbreeScene = FEmbreeScene();
+}
+
+const FEmbreeGeometryAsset* FEmbreeScene::AddGeometryAsset(
+	const FSourceMeshDataForDerivedDataTask& SourceMeshData,
+	const FStaticMeshLODResources* LODModel,
+	const TArray<FSignedDistanceFieldBuildSectionData>& SectionData,
+	bool bIncludeTranslucentTriangles,
+	bool bInstantiable)
+{
+#if USE_EMBREE
+	const uint32 NumVertices = SourceMeshData.IsValid() ? SourceMeshData.GetNumVertices() : LODModel->VertexBuffers.PositionVertexBuffer.GetNumVertices();
+	const uint32 NumIndices = SourceMeshData.IsValid() ? SourceMeshData.GetNumIndices() : LODModel->IndexBuffer.GetNumIndices();
+	const int32 NumTriangles = NumIndices / 3;
+
+	const FStaticMeshSectionArray& Sections = SourceMeshData.IsValid() ? SourceMeshData.Sections : LODModel->Sections;
 
 	TArray<int32> FilteredTriangles;
 	FilteredTriangles.Empty(NumTriangles);
 
+	// TODO: investigate/fix what causes this warning and then enable it
+	bool bWarnOnceSectionData = false;
+
 	for (int32 TriangleIndex = 0; TriangleIndex < NumTriangles; ++TriangleIndex)
 	{
 		FVector3f V0, V1, V2;
@@ -203,14 +219,14 @@ void MeshRepresentation::SetupEmbreeScene(
 		}
 		else
 		{
-			const FIndexArrayView Indices = LODModel.IndexBuffer.GetArrayView();
+			const FIndexArrayView Indices = LODModel->IndexBuffer.GetArrayView();
 			const uint32 I0 = Indices[TriangleIndex * 3 + 0];
 			const uint32 I1 = Indices[TriangleIndex * 3 + 1];
 			const uint32 I2 = Indices[TriangleIndex * 3 + 2];
 
-			V0 = LODModel.VertexBuffers.PositionVertexBuffer.VertexPosition(I0);
-			V1 = LODModel.VertexBuffers.PositionVertexBuffer.VertexPosition(I1);
-			V2 = LODModel.VertexBuffers.PositionVertexBuffer.VertexPosition(I2);
+			V0 = LODModel->VertexBuffers.PositionVertexBuffer.VertexPosition(I0);
+			V1 = LODModel->VertexBuffers.PositionVertexBuffer.VertexPosition(I1);
+			V2 = LODModel->VertexBuffers.PositionVertexBuffer.VertexPosition(I2);
 		}
 
 		const FVector3f TriangleNormal = ((V1 - V2) ^ (V0 - V2));
@@ -230,6 +246,11 @@ void MeshRepresentation::SetupEmbreeScene(
 						const bool bIsOpaqueOrMasked = !IsTranslucentBlendMode(SectionData[SectionIndex].BlendMode);
 						bIncludeTriangle = (bIsOpaqueOrMasked || bIncludeTranslucentTriangles) && SectionData[SectionIndex].bAffectDistanceFieldLighting;
 					}
+					else if (bWarnOnceSectionData)
+					{
+						UE_LOG(LogMeshUtilities, Warning, TEXT("Missing section data for %s, section = %d."), *MeshName, SectionIndex);
+						bWarnOnceSectionData = false;
+					}
 
 					break;
 				}
@@ -242,18 +263,22 @@ void MeshRepresentation::SetupEmbreeScene(
 		}
 	}
 
+	FEmbreeGeometryAsset& GeometryAsset = *new FEmbreeGeometryAsset;
+
 	const int32 NumBufferVerts = 1; // Reserve extra space at the end of the array, as embree has an internal bug where they read and discard 4 bytes off the end of the array
-	EmbreeScene.Geometry.VertexArray.Empty(NumVertices + NumBufferVerts);
-	EmbreeScene.Geometry.VertexArray.AddUninitialized(NumVertices + NumBufferVerts);
+
+	GeometryAsset.VertexArray.Empty(NumVertices + NumBufferVerts);
+	GeometryAsset.VertexArray.AddUninitialized(NumVertices + NumBufferVerts);
 
 	const int32 NumFilteredIndices = FilteredTriangles.Num() * 3;
 
-	EmbreeScene.Geometry.IndexArray.Empty(NumFilteredIndices);
-	EmbreeScene.Geometry.IndexArray.AddUninitialized(NumFilteredIndices);
+	GeometryAsset.IndexArray.Empty(NumFilteredIndices);
+	GeometryAsset.IndexArray.AddUninitialized(NumFilteredIndices);
+
+	GeometryAsset.TriangleDescs.Empty(FilteredTriangles.Num());
 
-	FVector3f* EmbreeVertices = EmbreeScene.Geometry.VertexArray.GetData();
-	uint32* EmbreeIndices = EmbreeScene.Geometry.IndexArray.GetData();
-	EmbreeScene.Geometry.TriangleDescs.Empty(FilteredTriangles.Num());
+	FVector3f* EmbreeVertices = GeometryAsset.VertexArray.GetData();
+	uint32* EmbreeIndices = GeometryAsset.IndexArray.GetData();
 
 	for (int32 FilteredTriangleIndex = 0; FilteredTriangleIndex < FilteredTriangles.Num(); FilteredTriangleIndex++)
 	{
@@ -273,14 +298,14 @@ void MeshRepresentation::SetupEmbreeScene(
 		}
 		else
 		{
-			const FIndexArrayView Indices = LODModel.IndexBuffer.GetArrayView();
+			const FIndexArrayView Indices = LODModel->IndexBuffer.GetArrayView();
 			I0 = Indices[TriangleIndex * 3 + 0];
 			I1 = Indices[TriangleIndex * 3 + 1];
 			I2 = Indices[TriangleIndex * 3 + 2];
 
-			V0 = LODModel.VertexBuffers.PositionVertexBuffer.VertexPosition(I0);
-			V1 = LODModel.VertexBuffers.PositionVertexBuffer.VertexPosition(I1);
-			V2 = LODModel.VertexBuffers.PositionVertexBuffer.VertexPosition(I2);
+			V0 = LODModel->VertexBuffers.PositionVertexBuffer.VertexPosition(I0);
+			V1 = LODModel->VertexBuffers.PositionVertexBuffer.VertexPosition(I1);
+			V2 = LODModel->VertexBuffers.PositionVertexBuffer.VertexPosition(I2);
 		}
 
 		bool bTriangleIsTwoSided = false;
@@ -295,6 +320,11 @@ void MeshRepresentation::SetupEmbreeScene(
 				{
 					bTriangleIsTwoSided = SectionData[SectionIndex].bTwoSided;
 				}
+				else if (bWarnOnceSectionData)
+				{
+					UE_LOG(LogMeshUtilities, Warning, TEXT("Missing section data for %s, section = %d."), *MeshName, SectionIndex);
+					bWarnOnceSectionData = false;
+				}
 
 				break;
 			}
@@ -312,70 +342,193 @@ void MeshRepresentation::SetupEmbreeScene(
 			FEmbreeTriangleDesc Desc;
 			// Store bGenerateAsIfTwoSided in material index
 			Desc.ElementIndex = bGenerateAsIfTwoSided || bTriangleIsTwoSided ? 1 : 0;
-			EmbreeScene.Geometry.TriangleDescs.Add(Desc);
+			GeometryAsset.TriangleDescs.Add(Desc);
 		}
 	}
 
-#if USE_EMBREE
 	{
-		RTCGeometry Geometry = rtcNewGeometry(EmbreeScene.EmbreeDevice, RTC_GEOMETRY_TYPE_TRIANGLE);
-		EmbreeScene.Geometry.InternalGeometry = Geometry;
+		GeometryAsset.NumVertices = NumVertices;
+		GeometryAsset.NumTriangles = FilteredTriangles.Num();
+
+		GeometryAsset.SectionNumTriangles = 0;
+		GeometryAsset.SectionNumTwoSidedTriangles = 0;
+
+		for (int32 SectionIndex = 0; SectionIndex < Sections.Num(); SectionIndex++)
+		{
+			const FStaticMeshSection& Section = Sections[SectionIndex];
+
+			if (SectionData.IsValidIndex(SectionIndex))
+			{
+				GeometryAsset.SectionNumTriangles += Section.NumTriangles;
+
+				if (SectionData[SectionIndex].bTwoSided)
+				{
+					GeometryAsset.SectionNumTwoSidedTriangles += Section.NumTriangles;
+				}
+			}
+			else if (bWarnOnceSectionData)
+			{
+				UE_LOG(LogMeshUtilities, Warning, TEXT("Missing section data for %s, section = %d."), *MeshName, SectionIndex);
+				bWarnOnceSectionData = false;
+			}
+		}
+	}
+
+	if (bInstantiable)
+	{
+		RTCGeometry ImplGeometry = rtcNewGeometry(Device, RTC_GEOMETRY_TYPE_TRIANGLE);
+
+		rtcSetSharedGeometryBuffer(ImplGeometry, RTC_BUFFER_TYPE_VERTEX, 0, RTC_FORMAT_FLOAT3, EmbreeVertices, 0, sizeof(FVector3f), GeometryAsset.NumVertices);
+		rtcSetSharedGeometryBuffer(ImplGeometry, RTC_BUFFER_TYPE_INDEX, 0, RTC_FORMAT_UINT3, EmbreeIndices, 0, sizeof(uint32) * 3, GeometryAsset.NumTriangles);
+
+		rtcSetGeometryUserData(ImplGeometry, &GeometryAsset);
+		rtcSetGeometryIntersectFilterFunction(ImplGeometry, EmbreeFilterFunc);
 
-		rtcSetSharedGeometryBuffer(Geometry, RTC_BUFFER_TYPE_VERTEX, 0, RTC_FORMAT_FLOAT3, EmbreeVertices, 0, sizeof(FVector3f), NumVertices);
-		rtcSetSharedGeometryBuffer(Geometry, RTC_BUFFER_TYPE_INDEX, 0, RTC_FORMAT_UINT3, EmbreeIndices, 0, sizeof(uint32) * 3, FilteredTriangles.Num());
+		rtcCommitGeometry(ImplGeometry);
 
-		rtcSetGeometryUserData(Geometry, &EmbreeScene.Geometry);
-		rtcSetGeometryIntersectFilterFunction(Geometry, EmbreeFilterFunc);
+		GeometryAsset.Scene = rtcNewScene(Device);
 
-		rtcCommitGeometry(Geometry);
-		rtcAttachGeometry(EmbreeScene.EmbreeScene, Geometry);
-		rtcReleaseGeometry(Geometry);
+		rtcAttachGeometry(GeometryAsset.Scene, ImplGeometry);
+		rtcReleaseGeometry(ImplGeometry);
 
-		rtcCommitScene(EmbreeScene.EmbreeScene);
+		rtcCommitScene(GeometryAsset.Scene);
 
-		RTCError ReturnError = rtcGetDeviceError(EmbreeScene.EmbreeDevice);
+		RTCError ReturnError = rtcGetDeviceError(Device);
 		if (ReturnError == RTC_ERROR_OUT_OF_MEMORY)
 		{
-			UE_LOG(LogMeshUtilities, Warning, TEXT("GenerateSignedDistanceFieldVolumeData failed for %s. Embree rtcCommitScene failed to allocate memory."), *MeshName);
+			UE_LOG(LogMeshUtilities, Warning, TEXT("Failed to create instantiable Embree geometry for %s (OUT_OF_MEMORY)."), *MeshName);
 			FPlatformMemory::OnOutOfMemory(0, 16);
-			return;
+			//return nullptr; // unreachable code
 		}
 		if (ReturnError != RTC_ERROR_NONE)
 		{
-			UE_LOG(LogMeshUtilities, Warning, TEXT("GenerateSignedDistanceFieldVolumeData failed for %s. Embree rtcCommitScene failed. Code: %d"), *MeshName, (int32)ReturnError);
-			return;
+			UE_LOG(LogMeshUtilities, Warning, TEXT("Failed to create instantiable Embree geometry for %s. Code: %d"), *MeshName, (int32)ReturnError);
+			return nullptr;
 		}
 	}
+
+	GeometryAssets.Add(&GeometryAsset);
+
+	return &GeometryAsset;
+#else
+	return nullptr;
 #endif
+}
+
+const FEmbreeGeometry* FEmbreeScene::AddGeometry(const FEmbreeGeometryAsset* GeometryAsset)
+{
+#if USE_EMBREE
+	const FVector3f* VerticesData = GeometryAsset->VertexArray.GetData();
+	const uint32* IndicesData = GeometryAsset->IndexArray.GetData();
+
+	RTCGeometry ImplGeometry = rtcNewGeometry(Device, RTC_GEOMETRY_TYPE_TRIANGLE);
+
+	rtcSetSharedGeometryBuffer(ImplGeometry, RTC_BUFFER_TYPE_VERTEX, 0, RTC_FORMAT_FLOAT3, VerticesData, 0, sizeof(FVector3f), GeometryAsset->NumVertices);
+	rtcSetSharedGeometryBuffer(ImplGeometry, RTC_BUFFER_TYPE_INDEX, 0, RTC_FORMAT_UINT3, IndicesData, 0, sizeof(uint32) * 3, GeometryAsset->NumTriangles);
 
-	// bMostlyTwoSided
+	rtcSetGeometryUserData(ImplGeometry, (void*)GeometryAsset);
+	rtcSetGeometryIntersectFilterFunction(ImplGeometry, EmbreeFilterFunc);
+
+	rtcCommitGeometry(ImplGeometry);
+
+	uint32 GeometryId = rtcAttachGeometry(Scene, ImplGeometry);
+	rtcReleaseGeometry(ImplGeometry);
+
+	RTCError ReturnError = rtcGetDeviceError(Device);
+	if (ReturnError == RTC_ERROR_OUT_OF_MEMORY)
+	{
+		UE_LOG(LogMeshUtilities, Warning, TEXT("Failed to add geometry to Embree scene for %s (OUT_OF_MEMORY)."), *MeshName);
+		FPlatformMemory::OnOutOfMemory(0, 16);
+		//return nullptr; // unreachable code
+	}
+	if (ReturnError != RTC_ERROR_NONE)
 	{
-		uint32 NumTrianglesTotal = 0;
-		uint32 NumTwoSidedTriangles = 0;
+		UE_LOG(LogMeshUtilities, Warning, TEXT("Failed to add geometry to Embree scene for %s. Code: %d"), *MeshName, (int32)ReturnError);
+		return nullptr;
+	}
 
-		for (int32 SectionIndex = 0; SectionIndex < Sections.Num(); SectionIndex++)
-		{
-			const FStaticMeshSection& Section = Sections[SectionIndex];
+	FEmbreeGeometry& Geometry = *new FEmbreeGeometry;
+	Geometry.Asset = GeometryAsset;
+	Geometry.GeometryId = GeometryId;
 
-			if (SectionData.IsValidIndex(SectionIndex))
-			{
-				NumTrianglesTotal += Section.NumTriangles;
+	Geometries.Add(&Geometry);
 
-				if (SectionData[SectionIndex].bTwoSided)
-				{
-					NumTwoSidedTriangles += Section.NumTriangles;
-				}
-			}
-		}
+	return &Geometry;
+#else
+	return nullptr;
+#endif
+}
+
+const FEmbreeGeometry* FEmbreeScene::AddGeometryInstance(const FEmbreeGeometryAsset* GeometryAsset, const FMatrix44f& Transform)
+{
+#if USE_EMBREE
+
+	RTCGeometry ImplGeometry = rtcNewGeometry(Device, RTC_GEOMETRY_TYPE_INSTANCE);
+
+	rtcSetGeometryInstancedScene(ImplGeometry, GeometryAsset->Scene);
+	rtcSetGeometryTransform(ImplGeometry, 0, RTC_FORMAT_FLOAT4X4_COLUMN_MAJOR, (const float*)Transform.M);
+	rtcCommitGeometry(ImplGeometry);
 
-		EmbreeScene.bMostlyTwoSided = NumTwoSidedTriangles * 4 >= NumTrianglesTotal || bGenerateAsIfTwoSided;
+	uint32 GeometryId = rtcAttachGeometry(Scene, ImplGeometry);
+	rtcReleaseGeometry(ImplGeometry);
+
+	RTCError ReturnError = rtcGetDeviceError(Device);
+	if (ReturnError == RTC_ERROR_OUT_OF_MEMORY)
+	{
+		UE_LOG(LogMeshUtilities, Warning, TEXT("Failed to add geometry instance to Embree scene for %s (OUT_OF_MEMORY)."), *MeshName);
+		FPlatformMemory::OnOutOfMemory(0, 16);
+		//return nullptr; // unreachable code
+	}
+	if (ReturnError != RTC_ERROR_NONE)
+	{
+		UE_LOG(LogMeshUtilities, Warning, TEXT("Failed to add geometry instance to Embree scene for %s. Code: %d"), *MeshName, (int32)ReturnError);
+		return nullptr;
 	}
+
+	FEmbreeGeometry& Geometry = *new FEmbreeGeometry;
+	Geometry.Asset = GeometryAsset;
+	Geometry.GeometryId = GeometryId;
+
+	Geometries.Add(&Geometry);
+
+	return &Geometry;
+#else
+	return nullptr;
+#endif
 }
 
-void MeshRepresentation::DeleteEmbreeScene(FEmbreeScene& EmbreeScene)
+void FEmbreeScene::Commit()
 {
 #if USE_EMBREE
-	rtcReleaseScene(EmbreeScene.EmbreeScene);
-	rtcReleaseDevice(EmbreeScene.EmbreeDevice);
+
+	NumTrianglesTotal = 0;
+
+	uint32 SectionNumTwoSidedTriangles = 0;
+	uint32 SectionNumTriangles = 0;
+
+	for (FEmbreeGeometry& Geometry : Geometries)
+	{
+		NumTrianglesTotal += Geometry.Asset->NumTriangles;
+
+		SectionNumTwoSidedTriangles += Geometry.Asset->SectionNumTwoSidedTriangles;
+		SectionNumTriangles += Geometry.Asset->SectionNumTriangles;
+	}
+
+	bMostlyTwoSided = SectionNumTwoSidedTriangles * 4 >= SectionNumTriangles || bGenerateAsIfTwoSided;
+
+	rtcCommitScene(Scene);
+
+	RTCError ReturnError = rtcGetDeviceError(Device);
+	if (ReturnError == RTC_ERROR_OUT_OF_MEMORY)
+	{
+		UE_LOG(LogMeshUtilities, Warning, TEXT("Failed to commit Embree scene for %s (OUT_OF_MEMORY)."), *MeshName);
+		FPlatformMemory::OnOutOfMemory(0, 16);
+		return;
+	}
+	if (ReturnError != RTC_ERROR_NONE)
+	{
+		UE_LOG(LogMeshUtilities, Warning, TEXT("Failed to commit Embree scene for %s. Code: %d"), *MeshName, (int32)ReturnError);
+		return;
+	}
 #endif
 }
\ No newline at end of file
diff --git a/Engine/Source/Developer/MeshUtilities/Private/MeshRepresentationCommon.h b/Engine/Source/Developer/MeshUtilities/Private/MeshRepresentationCommon.h
--- a/Engine/Source/Developer/MeshUtilities/Private/MeshRepresentationCommon.h
+++ b/Engine/Source/Developer/MeshUtilities/Private/MeshRepresentationCommon.h
@@ -26,24 +26,55 @@ struct FEmbreeTriangleDesc
 	}
 };
 
-// Mapping between Embree Geometry Id and engine Mesh/LOD Id
-struct FEmbreeGeometry
+struct FEmbreeGeometryAsset
 {
 	TArray<uint32> IndexArray;
 	TArray<FVector3f> VertexArray;
 	TArray<FEmbreeTriangleDesc> TriangleDescs; // The material ID of each triangle.
-	RTCGeometry InternalGeometry;
+
+	uint32 NumVertices;
+	uint32 NumTriangles;
+
+	uint32 SectionNumTwoSidedTriangles;
+	uint32 SectionNumTriangles;
+
+	RTCScene Scene = nullptr;
+};
+
+struct FEmbreeGeometry
+{
+	const FEmbreeGeometryAsset* Asset;
+	uint32 GeometryId = RTC_INVALID_GEOMETRY_ID;
 };
 
 class FEmbreeScene
 {
 public:
-	int32 NumIndices = 0;
-	bool bMostlyTwoSided = false;
 
-	RTCDevice EmbreeDevice = nullptr;
-	RTCScene EmbreeScene = nullptr;
-	FEmbreeGeometry Geometry;
+	const FEmbreeGeometryAsset* AddGeometryAsset(
+		const FSourceMeshDataForDerivedDataTask& SourceMeshData,
+		const FStaticMeshLODResources* LODModel,
+		const TArray<FSignedDistanceFieldBuildSectionData>& SectionData,
+		bool bIncludeTranslucentTriangles,
+		bool bInstantiable);
+
+	const FEmbreeGeometry* AddGeometry(const FEmbreeGeometryAsset* GeometryAsset);
+
+	const FEmbreeGeometry* AddGeometryInstance(const FEmbreeGeometryAsset* GeometryAsset, const FMatrix44f& Transform);
+
+	void Commit();
+
+	FString MeshName;
+	bool bGenerateAsIfTwoSided = false;
+
+	RTCDevice Device = nullptr;
+	RTCScene Scene = nullptr;
+
+	TIndirectArray<FEmbreeGeometryAsset> GeometryAssets;
+	TIndirectArray<FEmbreeGeometry> Geometries;
+
+	int32 NumTrianglesTotal = 0;
+	bool bMostlyTwoSided = false;
 };
 
 #if USE_EMBREE
@@ -108,13 +139,8 @@ namespace MeshRepresentation
 	 */
 	FMatrix44f GetTangentBasisFrisvad(FVector3f TangentZ);
 
-	void SetupEmbreeScene(FString MeshName,
-		const FSourceMeshDataForDerivedDataTask& SourceMeshData,
-		const FStaticMeshLODResources& LODModel,
-		const TArray<FSignedDistanceFieldBuildSectionData>& SectionData,
-		bool bGenerateAsIfTwoSided,
-		bool bIncludeTranslucentTriangles,
-		FEmbreeScene& EmbreeScene);
+	void SetupEmbreeScene(FString MeshName, bool bGenerateAsIfTwoSided, FEmbreeScene& OutEmbreeScene);
 
 	void DeleteEmbreeScene(FEmbreeScene& EmbreeScene);
+
 };
diff --git a/Engine/Source/Runtime/Engine/Private/DistanceFieldAtlas.cpp b/Engine/Source/Runtime/Engine/Private/DistanceFieldAtlas.cpp
--- a/Engine/Source/Runtime/Engine/Private/DistanceFieldAtlas.cpp
+++ b/Engine/Source/Runtime/Engine/Private/DistanceFieldAtlas.cpp
@@ -128,7 +128,7 @@ FDistanceFieldAsyncQueue* GDistanceFieldAsyncQueue = NULL;
 #if WITH_EDITOR
 
 // DDC key for distance field data, must be changed when modifying the generation code or data format
-#define DISTANCEFIELD_DERIVEDDATA_VER TEXT("4BD0E980-7212-4093-A14D-2E468CD9FAF3")
+#define DISTANCEFIELD_DERIVEDDATA_VER TEXT("9137A53B-09F6-4C2F-86E1-29FE6834D74B")
 
 FString BuildDistanceFieldDerivedDataKey(const FString& InMeshKey)
 {
diff --git a/Engine/Source/Runtime/Engine/Private/MeshCardRepresentation.cpp b/Engine/Source/Runtime/Engine/Private/MeshCardRepresentation.cpp
--- a/Engine/Source/Runtime/Engine/Private/MeshCardRepresentation.cpp
+++ b/Engine/Source/Runtime/Engine/Private/MeshCardRepresentation.cpp
@@ -101,7 +101,7 @@ FCardRepresentationAsyncQueue* GCardRepresentationAsyncQueue = NULL;
 #if WITH_EDITOR
 
 // DDC key for card representation data, must be changed when modifying the generation code or data format
-#define CARDREPRESENTATION_DERIVEDDATA_VER TEXT("B58BB068-C040-4A1D-BC7C-A06D36FE3B1D")
+#define CARDREPRESENTATION_DERIVEDDATA_VER TEXT("59BE661D-D92C-4158-B14D-3C384D5F5FEF")
 
 FString BuildCardRepresentationDerivedDataKey(const FString& InMeshKey, int32 MaxLumenMeshCards)
 {
