Refactor some MeshUtilities logic.
- pass some of GenerateSignedDistanceFieldVolumeData(...) and GenerateCardRepresentationData(...) args using a struct.
- helper function MeshRepresentation::AddMeshDataToEmbreeScene(...)

diff --git a/Engine/Source/Developer/MeshUtilities/Private/MeshCardRepresentationUtilities.cpp b/Engine/Source/Developer/MeshUtilities/Private/MeshCardRepresentationUtilities.cpp
--- a/Engine/Source/Developer/MeshUtilities/Private/MeshCardRepresentationUtilities.cpp
+++ b/Engine/Source/Developer/MeshUtilities/Private/MeshCardRepresentationUtilities.cpp
@@ -1047,11 +1047,7 @@ void BuildMeshCards(const FBox& MeshBounds, const FGenerateCardMeshContext& Cont
 
 bool FMeshUtilities::GenerateCardRepresentationData(
 	FString MeshName,
-	const FSourceMeshDataForDerivedDataTask& SourceMeshData,
-	const FStaticMeshLODResources& LODModel,
-	class FQueuedThreadPool& ThreadPool,
-	const TArray<FSignedDistanceFieldBuildSectionData>& SectionData,
-	const FBoxSphereBounds& Bounds,
+	const FMeshDataForDerivedDataTask& MeshData,
 	const FDistanceFieldVolumeData* DistanceFieldVolumeData,
 	int32 MaxLumenMeshCards,
 	bool bGenerateAsIfTwoSided,
@@ -1076,17 +1072,18 @@ bool FMeshUtilities::GenerateCardRepresentationData(
 			return false;
 		}
 
-		const FEmbreeGeometryAsset* GeometryAsset = EmbreeScene.AddGeometryAsset(SourceMeshData, &LODModel, SectionData, bIncludeTranslucentTriangles, false);
-
-		const FEmbreeGeometry* Geometry = EmbreeScene.AddGeometry(GeometryAsset);
+		if (!MeshRepresentation::AddMeshDataToEmbreeScene(EmbreeScene, MeshData, bIncludeTranslucentTriangles))
+		{
+			return false;
+		}
 
 		EmbreeScene.Commit();
 
 		FGenerateCardMeshContext Context(MeshName, EmbreeScene, OutData);
 
 		// Note: must operate on the SDF bounds when available, because SDF generation can expand the mesh's bounds
-		const FBox BuildCardsBounds = DistanceFieldVolumeData && DistanceFieldVolumeData->LocalSpaceMeshBounds.IsValid ? FBox(DistanceFieldVolumeData->LocalSpaceMeshBounds) : Bounds.GetBox();
-		BuildMeshCards(BuildCardsBounds, Context, MaxLumenMeshCards, OutData);
+		const FBox3f BuildCardsBounds = DistanceFieldVolumeData && DistanceFieldVolumeData->LocalSpaceMeshBounds.IsValid ? DistanceFieldVolumeData->LocalSpaceMeshBounds : MeshData.Bounds.GetBox();
+		BuildMeshCards(FBox(BuildCardsBounds), Context, MaxLumenMeshCards, OutData);
 
 		MeshRepresentation::DeleteEmbreeScene(EmbreeScene);
 
diff --git a/Engine/Source/Developer/MeshUtilities/Private/MeshDistanceFieldUtilities.cpp b/Engine/Source/Developer/MeshUtilities/Private/MeshDistanceFieldUtilities.cpp
--- a/Engine/Source/Developer/MeshUtilities/Private/MeshDistanceFieldUtilities.cpp
+++ b/Engine/Source/Developer/MeshUtilities/Private/MeshDistanceFieldUtilities.cpp
@@ -515,14 +515,10 @@ static void BuildSignedDistanceField(
 
 void FMeshUtilities::GenerateSignedDistanceFieldVolumeData(
 	FString MeshName,
-	const FSourceMeshDataForDerivedDataTask& SourceMeshData,
-	const FStaticMeshLODResources& LODModel,
-	class FQueuedThreadPool& ThreadPool,
-	const TArray<FSignedDistanceFieldBuildSectionData>& SectionData,
-	const FBoxSphereBounds3f& Bounds,
+	const FMeshDataForDerivedDataTask& MeshData,
 	float DistanceFieldResolutionScale,
 	bool bGenerateAsIfTwoSided,
-	FDistanceFieldVolumeData& OutData)
+	class FDistanceFieldVolumeData& OutData)
 {
 	TRACE_CPUPROFILER_EVENT_SCOPE(GenerateSignedDistanceFieldVolumeData);
 	COOK_STAT(auto Timer = DistanceFieldUtilityCookStats::UsageStats.TimeSyncWork());
@@ -542,13 +538,14 @@ void FMeshUtilities::GenerateSignedDistanceFieldVolumeData(
 			return;
 		}
 
-		const FEmbreeGeometryAsset* GeometryAsset = EmbreeScene.AddGeometryAsset(SourceMeshData, &LODModel, SectionData, bIncludeTranslucentTriangles, false);
-
-		const FEmbreeGeometry* Geometry = EmbreeScene.AddGeometry(GeometryAsset);
+		if (!MeshRepresentation::AddMeshDataToEmbreeScene(EmbreeScene, MeshData, bIncludeTranslucentTriangles))
+		{
+			return;
+		}
 
 		EmbreeScene.Commit();
 
-		BuildSignedDistanceField(MeshName, EmbreeScene, Bounds, DistanceFieldResolutionScale, OutData);
+		BuildSignedDistanceField(MeshName, EmbreeScene, MeshData.Bounds, DistanceFieldResolutionScale, OutData);
 
 		MeshRepresentation::DeleteEmbreeScene(EmbreeScene);
 
@@ -576,14 +573,10 @@ void FMeshUtilities::GenerateSignedDistanceFieldVolumeData(
 
 void FMeshUtilities::GenerateSignedDistanceFieldVolumeData(
 	FString MeshName,
-	const FSourceMeshDataForDerivedDataTask& SourceMeshData,
-	const FStaticMeshLODResources& LODModel,
-	class FQueuedThreadPool& ThreadPool,
-	const TArray<FSignedDistanceFieldBuildSectionData>& SectionData,
-	const FBoxSphereBounds& Bounds,
+	const FMeshDataForDerivedDataTask& MeshData,
 	float DistanceFieldResolutionScale,
 	bool bGenerateAsIfTwoSided,
-	FDistanceFieldVolumeData& OutData)
+	class FDistanceFieldVolumeData& OutData)
 {
 	if (DistanceFieldResolutionScale > 0)
 	{
diff --git a/Engine/Source/Developer/MeshUtilities/Private/MeshRepresentationCommon.cpp b/Engine/Source/Developer/MeshUtilities/Private/MeshRepresentationCommon.cpp
--- a/Engine/Source/Developer/MeshUtilities/Private/MeshRepresentationCommon.cpp
+++ b/Engine/Source/Developer/MeshUtilities/Private/MeshRepresentationCommon.cpp
@@ -183,19 +183,36 @@ void MeshRepresentation::DeleteEmbreeScene(FEmbreeScene& EmbreeScene)
 	EmbreeScene = FEmbreeScene();
 }
 
+bool MeshRepresentation::AddMeshDataToEmbreeScene(FEmbreeScene& EmbreeScene, const FMeshDataForDerivedDataTask& MeshData, bool bIncludeTranslucentTriangles)
+{
+	if ((MeshData.SourceMeshData == nullptr || !MeshData.SourceMeshData->IsValid()) && MeshData.LODModel == nullptr)
+	{
+		UE_LOG(LogMeshUtilities, Warning, TEXT("Provided MeshData for %s doesn't contain any data."), *EmbreeScene.MeshName);
+		return false;
+	}
+
+	const FEmbreeGeometryAsset* GeometryAsset = EmbreeScene.AddGeometryAsset(MeshData.SourceMeshData, MeshData.LODModel, MeshData.SectionData, bIncludeTranslucentTriangles, false);
+
+	const FEmbreeGeometry* Geometry = EmbreeScene.AddGeometry(GeometryAsset);
+
+	return true;
+}
+
 const FEmbreeGeometryAsset* FEmbreeScene::AddGeometryAsset(
-	const FSourceMeshDataForDerivedDataTask& SourceMeshData,
+	const FSourceMeshDataForDerivedDataTask* SourceMeshData,
 	const FStaticMeshLODResources* LODModel,
-	const TArray<FSignedDistanceFieldBuildSectionData>& SectionData,
+	TConstArrayView<FSignedDistanceFieldBuildSectionData> SectionData,
 	bool bIncludeTranslucentTriangles,
 	bool bInstantiable)
 {
+	check((SourceMeshData != nullptr && SourceMeshData->IsValid()) || (LODModel != nullptr));
+
 #if USE_EMBREE
-	const uint32 NumVertices = SourceMeshData.IsValid() ? SourceMeshData.GetNumVertices() : LODModel->VertexBuffers.PositionVertexBuffer.GetNumVertices();
-	const uint32 NumIndices = SourceMeshData.IsValid() ? SourceMeshData.GetNumIndices() : LODModel->IndexBuffer.GetNumIndices();
+	const uint32 NumVertices = (SourceMeshData && SourceMeshData->IsValid()) ? SourceMeshData->GetNumVertices() : LODModel->VertexBuffers.PositionVertexBuffer.GetNumVertices();
+	const uint32 NumIndices = (SourceMeshData && SourceMeshData->IsValid()) ? SourceMeshData->GetNumIndices() : LODModel->IndexBuffer.GetNumIndices();
 	const int32 NumTriangles = NumIndices / 3;
 
-	const FStaticMeshSectionArray& Sections = SourceMeshData.IsValid() ? SourceMeshData.Sections : LODModel->Sections;
+	const FStaticMeshSectionArray& Sections = (SourceMeshData && SourceMeshData->IsValid()) ? SourceMeshData->Sections : LODModel->Sections;
 
 	TArray<int32> FilteredTriangles;
 	FilteredTriangles.Empty(NumTriangles);
@@ -207,15 +224,15 @@ const FEmbreeGeometryAsset* FEmbreeScene::AddGeometryAsset(
 	{
 		FVector3f V0, V1, V2;
 
-		if (SourceMeshData.IsValid())
+		if (SourceMeshData && SourceMeshData->IsValid())
 		{
-			const uint32 I0 = SourceMeshData.TriangleIndices[TriangleIndex * 3 + 0];
-			const uint32 I1 = SourceMeshData.TriangleIndices[TriangleIndex * 3 + 1];
-			const uint32 I2 = SourceMeshData.TriangleIndices[TriangleIndex * 3 + 2];
+			const uint32 I0 = SourceMeshData->TriangleIndices[TriangleIndex * 3 + 0];
+			const uint32 I1 = SourceMeshData->TriangleIndices[TriangleIndex * 3 + 1];
+			const uint32 I2 = SourceMeshData->TriangleIndices[TriangleIndex * 3 + 2];
 
-			V0 = SourceMeshData.VertexPositions[I0];
-			V1 = SourceMeshData.VertexPositions[I1];
-			V2 = SourceMeshData.VertexPositions[I2];
+			V0 = SourceMeshData->VertexPositions[I0];
+			V1 = SourceMeshData->VertexPositions[I1];
+			V2 = SourceMeshData->VertexPositions[I2];
 		}
 		else
 		{
@@ -286,15 +303,15 @@ const FEmbreeGeometryAsset* FEmbreeScene::AddGeometryAsset(
 		FVector3f V0, V1, V2;
 
 		const int32 TriangleIndex = FilteredTriangles[FilteredTriangleIndex];
-		if (SourceMeshData.IsValid())
+		if (SourceMeshData && SourceMeshData->IsValid())
 		{
-			I0 = SourceMeshData.TriangleIndices[TriangleIndex * 3 + 0];
-			I1 = SourceMeshData.TriangleIndices[TriangleIndex * 3 + 1];
-			I2 = SourceMeshData.TriangleIndices[TriangleIndex * 3 + 2];
+			I0 = SourceMeshData->TriangleIndices[TriangleIndex * 3 + 0];
+			I1 = SourceMeshData->TriangleIndices[TriangleIndex * 3 + 1];
+			I2 = SourceMeshData->TriangleIndices[TriangleIndex * 3 + 2];
 
-			V0 = SourceMeshData.VertexPositions[I0];
-			V1 = SourceMeshData.VertexPositions[I1];
-			V2 = SourceMeshData.VertexPositions[I2];
+			V0 = SourceMeshData->VertexPositions[I0];
+			V1 = SourceMeshData->VertexPositions[I1];
+			V2 = SourceMeshData->VertexPositions[I2];
 		}
 		else
 		{
diff --git a/Engine/Source/Developer/MeshUtilities/Private/MeshRepresentationCommon.h b/Engine/Source/Developer/MeshUtilities/Private/MeshRepresentationCommon.h
--- a/Engine/Source/Developer/MeshUtilities/Private/MeshRepresentationCommon.h
+++ b/Engine/Source/Developer/MeshUtilities/Private/MeshRepresentationCommon.h
@@ -52,9 +52,9 @@ class FEmbreeScene
 public:
 
 	const FEmbreeGeometryAsset* AddGeometryAsset(
-		const FSourceMeshDataForDerivedDataTask& SourceMeshData,
+		const FSourceMeshDataForDerivedDataTask* SourceMeshData,
 		const FStaticMeshLODResources* LODModel,
-		const TArray<FSignedDistanceFieldBuildSectionData>& SectionData,
+		TConstArrayView<FSignedDistanceFieldBuildSectionData> SectionData,
 		bool bIncludeTranslucentTriangles,
 		bool bInstantiable);
 
@@ -143,4 +143,6 @@ namespace MeshRepresentation
 
 	void DeleteEmbreeScene(FEmbreeScene& EmbreeScene);
 
+	bool AddMeshDataToEmbreeScene(FEmbreeScene& EmbreeScene, const FMeshDataForDerivedDataTask& MeshData, bool bIncludeTranslucentTriangles);
+
 };
diff --git a/Engine/Source/Developer/MeshUtilities/Private/MeshUtilitiesPrivate.h b/Engine/Source/Developer/MeshUtilities/Private/MeshUtilitiesPrivate.h
--- a/Engine/Source/Developer/MeshUtilities/Private/MeshUtilitiesPrivate.h
+++ b/Engine/Source/Developer/MeshUtilities/Private/MeshUtilitiesPrivate.h
@@ -55,24 +55,17 @@ private:
 
 	virtual bool GenerateStaticMeshLODs(UStaticMesh* StaticMesh, const FStaticMeshLODGroup& LODGroup) override;
 
+	/** Builds a signed distance field volume for the given mesh. */
 	virtual void GenerateSignedDistanceFieldVolumeData(
 		FString MeshName,
-		const FSourceMeshDataForDerivedDataTask& SourceMeshData,
-		const FStaticMeshLODResources& LODModel,
-		class FQueuedThreadPool& ThreadPool,
-		const TArray<FSignedDistanceFieldBuildSectionData>& SectionData,
-		const FBoxSphereBounds3f& Bounds,
+		const FMeshDataForDerivedDataTask& MeshData,
 		float DistanceFieldResolutionScale,
 		bool bGenerateAsIfTwoSided,
-		FDistanceFieldVolumeData& OutData) override;
+		class FDistanceFieldVolumeData& OutData) override;
 	
 	virtual bool GenerateCardRepresentationData(
 		FString MeshName,
-		const FSourceMeshDataForDerivedDataTask& SourceMeshData,
-		const FStaticMeshLODResources& LODModel,
-		class FQueuedThreadPool& ThreadPool,
-		const TArray<FSignedDistanceFieldBuildSectionData>& SectionData,
-		const FBoxSphereBounds& Bounds,
+		const FMeshDataForDerivedDataTask& MeshData,
 		const FDistanceFieldVolumeData* DistanceFieldVolumeData,
 		int32 MaxLumenMeshCards,
 		bool bGenerateAsIfTwoSided,
diff --git a/Engine/Source/Developer/MeshUtilities/Public/MeshUtilities.h b/Engine/Source/Developer/MeshUtilities/Public/MeshUtilities.h
--- a/Engine/Source/Developer/MeshUtilities/Public/MeshUtilities.h
+++ b/Engine/Source/Developer/MeshUtilities/Public/MeshUtilities.h
@@ -54,6 +54,14 @@ struct FSignedDistanceFieldBuildSectionData
 	bool bAffectDistanceFieldLighting = true;
 };
 
+struct FMeshDataForDerivedDataTask
+{
+	const FSourceMeshDataForDerivedDataTask* SourceMeshData;
+	const FStaticMeshLODResources* LODModel;
+	TConstArrayView<FSignedDistanceFieldBuildSectionData> SectionData;
+	FBoxSphereBounds3f Bounds;
+};
+
 struct FOverlappingCorners;
 
 class IMeshUtilities : public IModuleInterface
@@ -122,7 +130,15 @@ public:
 		const class FStaticMeshLODGroup& LODGroup
 		) = 0;
 
-	/** Builds a signed distance field volume for the given LODModel. */
+	/** Builds a signed distance field volume for the given mesh. */
+	virtual void GenerateSignedDistanceFieldVolumeData(
+		FString MeshName,
+		const FMeshDataForDerivedDataTask& MeshData,
+		float DistanceFieldResolutionScale,
+		bool bGenerateAsIfTwoSided,
+		class FDistanceFieldVolumeData& OutData) = 0;
+
+	UE_DEPRECATED(5.6, "Provide FMeshDataForDerivedDataTask instead.")
 	virtual void GenerateSignedDistanceFieldVolumeData(
 		FString MeshName,
 		const FSourceMeshDataForDerivedDataTask& SourceMeshData,
@@ -132,8 +148,26 @@ public:
 		const FBoxSphereBounds3f& Bounds,
 		float DistanceFieldResolutionScale,
 		bool bGenerateAsIfTwoSided,
-		class FDistanceFieldVolumeData& OutData) = 0;
+		class FDistanceFieldVolumeData& OutData)
+	{
+		FMeshDataForDerivedDataTask MeshData;
+		MeshData.SourceMeshData = &SourceMeshData;
+		MeshData.LODModel = &LODModel;
+		MeshData.SectionData = SectionData;
+		MeshData.Bounds = Bounds;
 
+		GenerateSignedDistanceFieldVolumeData(MeshName, MeshData, DistanceFieldResolutionScale, bGenerateAsIfTwoSided, OutData);
+	}
+
+	virtual bool GenerateCardRepresentationData(
+		FString MeshName,
+		const FMeshDataForDerivedDataTask& MeshData,
+		const class FDistanceFieldVolumeData* DistanceFieldVolumeData,
+		int32 MaxLumenMeshCards,
+		bool bGenerateAsIfTwoSided,
+		class FCardRepresentationData& OutData) = 0;
+
+	UE_DEPRECATED(5.6, "Provide FMeshDataForDerivedDataTask instead.")
 	virtual bool GenerateCardRepresentationData(
 		FString MeshName,
 		const FSourceMeshDataForDerivedDataTask& SourceMeshData,
@@ -144,7 +178,16 @@ public:
 		const class FDistanceFieldVolumeData* DistanceFieldVolumeData,
 		int32 MaxLumenMeshCards,
 		bool bGenerateAsIfTwoSided,
-		class FCardRepresentationData& OutData) = 0;
+		class FCardRepresentationData& OutData)
+	{
+		FMeshDataForDerivedDataTask MeshData;
+		MeshData.SourceMeshData = &SourceMeshData;
+		MeshData.LODModel = &LODModel;
+		MeshData.SectionData = SectionData;
+		MeshData.Bounds = (FBoxSphereBounds3f)Bounds;
+
+		return GenerateCardRepresentationData(MeshName, MeshData, DistanceFieldVolumeData, MaxLumenMeshCards, bGenerateAsIfTwoSided, OutData);
+	}
 
 	/** Helper structure for skeletal mesh import options */
 	struct MeshBuildOptions
diff --git a/Engine/Source/Runtime/Engine/Private/DistanceFieldAtlas.cpp b/Engine/Source/Runtime/Engine/Private/DistanceFieldAtlas.cpp
--- a/Engine/Source/Runtime/Engine/Private/DistanceFieldAtlas.cpp
+++ b/Engine/Source/Runtime/Engine/Private/DistanceFieldAtlas.cpp
@@ -738,13 +738,16 @@ void FDistanceFieldAsyncQueue::Build(FAsyncDistanceFieldTask* Task, FQueuedThrea
 		TRACE_CPUPROFILER_EVENT_SCOPE(FDistanceFieldAsyncQueue::Build);
 
 		const FStaticMeshLODResources& LODModel = Task->GenerateSource->GetRenderData()->LODResources[0];
+
+		FMeshDataForDerivedDataTask MeshData;
+		MeshData.SourceMeshData = &Task->SourceMeshData;
+		MeshData.LODModel = &LODModel;
+		MeshData.SectionData = Task->SectionData;
+		MeshData.Bounds = (FBoxSphereBounds3f)Task->GenerateSource->GetRenderData()->Bounds;
+
 		MeshUtilities.load(std::memory_order_relaxed)->GenerateSignedDistanceFieldVolumeData(
 			Task->StaticMesh->GetName(),
-			Task->SourceMeshData,
-			LODModel,
-			BuildThreadPool,
-			Task->SectionData,
-			(FBoxSphereBounds3f)Task->GenerateSource->GetRenderData()->Bounds,
+			MeshData,
 			Task->DistanceFieldResolutionScale,
 			Task->bGenerateDistanceFieldAsIfTwoSided,
 			*Task->GeneratedVolumeData);
diff --git a/Engine/Source/Runtime/Engine/Private/MeshCardRepresentation.cpp b/Engine/Source/Runtime/Engine/Private/MeshCardRepresentation.cpp
--- a/Engine/Source/Runtime/Engine/Private/MeshCardRepresentation.cpp
+++ b/Engine/Source/Runtime/Engine/Private/MeshCardRepresentation.cpp
@@ -679,13 +679,15 @@ void FCardRepresentationAsyncQueue::Build(FAsyncCardRepresentationTask* Task, FQ
 
 		const FStaticMeshLODResources& LODModel = Task->GenerateSource->GetRenderData()->LODResources[0];
 
+		FMeshDataForDerivedDataTask MeshData;
+		MeshData.SourceMeshData = &Task->SourceMeshData;
+		MeshData.LODModel = &LODModel;
+		MeshData.SectionData = Task->SectionData;
+		MeshData.Bounds = (FBoxSphereBounds3f)Task->GenerateSource->GetRenderData()->Bounds;
+
 		Task->bSuccess = MeshUtilities->GenerateCardRepresentationData(
 			Task->StaticMesh->GetName(),
-			Task->SourceMeshData,
-			LODModel,
-			BuildThreadPool,
-			Task->SectionData,
-			Task->GenerateSource->GetRenderData()->Bounds,
+			MeshData,
 			Task->GenerateSource->GetRenderData()->LODResources[0].DistanceFieldData,
 			Task->MaxLumenMeshCards,
 			Task->bGenerateDistanceFieldAsIfTwoSided,
