Add insights memory tracking to Embree

diff --git a/Engine/Source/Developer/MeshUtilities/Private/MeshRepresentationCommon.cpp b/Engine/Source/Developer/MeshUtilities/Private/MeshRepresentationCommon.cpp
--- a/Engine/Source/Developer/MeshUtilities/Private/MeshRepresentationCommon.cpp
+++ b/Engine/Source/Developer/MeshUtilities/Private/MeshRepresentationCommon.cpp
@@ -122,6 +122,20 @@ void EmbreeErrorFunc(void* userPtr, RTCError code, const char* str)
 
 	UE_LOG(LogMeshUtilities, Error, TEXT("Embree error: %s Code=%u"), *ErrorString, (uint32)code);
 }
+
+LLM_DEFINE_TAG(Embree);
+
+// userPtr is provided during callback registration, it points to the associated FEmbreeScene
+bool EmbreeMemoryMonitor(void* userPtr, ssize_t bytes, bool postAllocation)
+{
+	// postAllocation is true after an allocation has been made or freed.
+	if (postAllocation) {
+		LLM_SCOPE_BYTAG(Embree);
+		LLM_IF_ENABLED(FLowLevelMemTracker::Get().OnLowLevelChangeInMemoryUse(ELLMTracker::Default, static_cast<int64>(bytes)));
+	}
+
+	return true;
+}
 #endif
 
 void MeshRepresentation::SetupEmbreeScene(FString MeshName, bool bGenerateAsIfTwoSided, FEmbreeScene& OutEmbreeScene)
@@ -132,6 +146,7 @@ void MeshRepresentation::SetupEmbreeScene(FString MeshName, bool bGenerateAsIfTw
 #if USE_EMBREE
 	OutEmbreeScene.Device = rtcNewDevice(nullptr);
 	rtcSetDeviceErrorFunction(OutEmbreeScene.Device, EmbreeErrorFunc, nullptr);
+	LLM_IF_ENABLED(rtcSetDeviceMemoryMonitorFunction(OutEmbreeScene.Device, EmbreeMemoryMonitor, &OutEmbreeScene));
 
 	RTCError ReturnErrorNewDevice = rtcGetDeviceError(OutEmbreeScene.Device);
 	if (ReturnErrorNewDevice == RTC_ERROR_OUT_OF_MEMORY)
