Skeletal Mesh: Make sure we clear old skin weight and vertex attribute data when rebuilding LOD models
Otherwise the reduction code tries to make use of out-of-date versions of these data.

#rnx

diff --git a/Engine/Source/Developer/MeshUtilities/Private/MeshUtilities.cpp b/Engine/Source/Developer/MeshUtilities/Private/MeshUtilities.cpp
--- a/Engine/Source/Developer/MeshUtilities/Private/MeshUtilities.cpp
+++ b/Engine/Source/Developer/MeshUtilities/Private/MeshUtilities.cpp
@@ -796,6 +796,11 @@ void FMeshUtilities::BuildSkeletalModelFromChunks(
 	LODModel.NumVertices = 0;
 	LODModel.IndexBuffer.Empty();
 
+	// Since we're building a new model, clear out skin weight profiles and vertex attributes. Those can be recreated by
+	// FLODUtilities::UpdateAlternateSkinWeights & FLODUtilities::UpdateLODInfoVertexAttributes as needed.
+	LODModel.SkinWeightProfiles.Empty();
+	LODModel.VertexAttributes.Empty();
+	
 	// Setup the section and chunk arrays on the model.
 	for (int32 ChunkIndex = 0; ChunkIndex < Chunks.Num(); ++ChunkIndex)
 	{
