Add memory estimates for MeshDistanceField cook

#rnx

diff --git a/Engine/Source/Developer/MeshUtilities/Private/MeshDistanceFieldUtilities.cpp b/Engine/Source/Developer/MeshUtilities/Private/MeshDistanceFieldUtilities.cpp
--- a/Engine/Source/Developer/MeshUtilities/Private/MeshDistanceFieldUtilities.cpp
+++ b/Engine/Source/Developer/MeshUtilities/Private/MeshDistanceFieldUtilities.cpp
@@ -23,6 +23,12 @@ namespace DistanceFieldUtilityCookStats
 }
 #endif
 
+static TAutoConsoleVariable<int32> CVarDistanceFieldsEstimateRatio(
+	TEXT("r.CookMemory.DistanceFieldsEstimateRatio"), 15,
+	TEXT("Configurable ratio for the distance field memory estimate"),
+	ECVF_ReadOnly
+); 
+
 #if USE_EMBREE
 
 
@@ -290,6 +296,65 @@ void FSparseMeshDistanceFieldAsyncTask::DoWork()
 	}
 }
 
+int64_t FMeshUtilities::MemoryForGenerateSignedDistanceFieldVolumeData(
+	const FMeshDataForDerivedDataTask& MeshData,
+	float DistanceFieldResolutionScale
+)
+{
+	// ------------------
+	// Copied from GenerateSignedDistanceFieldVolumeData
+	// ------------------
+	static const auto CVar = IConsoleManager::Get().FindTConsoleVariableDataInt(TEXT("r.DistanceFields.MaxPerMeshResolution"));
+	const int32 PerMeshMax = CVar->GetValueOnAnyThread();
+
+	// Meshes with explicit artist-specified scale can go higher
+	const int32 MaxNumBlocksOneDim = FMath::Min<int32>(FMath::DivideAndRoundNearest(DistanceFieldResolutionScale <= 1 ? PerMeshMax / 2 : PerMeshMax, DistanceField::UniqueDataBrickSize), DistanceField::MaxIndirectionDimension - 1);
+
+	static const auto CVarDensity = IConsoleManager::Get().FindTConsoleVariableDataFloat(TEXT("r.DistanceFields.DefaultVoxelDensity"));
+	const float VoxelDensity = CVarDensity->GetValueOnAnyThread();
+
+	const float NumVoxelsPerLocalSpaceUnit = VoxelDensity * DistanceFieldResolutionScale;
+
+	FBox3f LocalSpaceMeshBounds(MeshData.Bounds.GetBox());
+
+	// Make sure the mesh bounding box has positive extents to handle planes
+	{
+		FVector3f MeshBoundsCenter = LocalSpaceMeshBounds.GetCenter();
+		FVector3f MeshBoundsExtent = FVector3f::Max(LocalSpaceMeshBounds.GetExtent(), FVector3f::OneVector);
+		LocalSpaceMeshBounds.Min = MeshBoundsCenter - MeshBoundsExtent;
+		LocalSpaceMeshBounds.Max = MeshBoundsCenter + MeshBoundsExtent;
+	}
+	const FVector3f DesiredDimensions = FVector3f(LocalSpaceMeshBounds.GetSize() * NumVoxelsPerLocalSpaceUnit / (float)DistanceField::UniqueDataBrickSize);
+
+	const FInt32Vector Mip0IndirectionDimensions = FInt32Vector(
+		FMath::Clamp(FMath::RoundToInt(DesiredDimensions.X), 1, MaxNumBlocksOneDim),
+		FMath::Clamp(FMath::RoundToInt(DesiredDimensions.Y), 1, MaxNumBlocksOneDim),
+		FMath::Clamp(FMath::RoundToInt(DesiredDimensions.Z), 1, MaxNumBlocksOneDim));
+
+	// ------------------
+
+	const int64 TotalBrickCount = (Mip0IndirectionDimensions.X * DistanceField::UniqueDataBrickSize 
+		* Mip0IndirectionDimensions.Y * DistanceField::UniqueDataBrickSize
+		* Mip0IndirectionDimensions.Z * DistanceField::UniqueDataBrickSize);
+
+	// Base and Ratio have been taken by calculating at a granular level the memory used and correlating 
+	// it with the TotalBrickCount. An estimate was used in place of granular calculations for simplicity.
+	const int32 Base = 1000;
+	const int32 Ratio = CVarDistanceFieldsEstimateRatio.GetValueOnAnyThread();
+
+	const int64 MeshUtilTotal = Base + TotalBrickCount * Ratio;
+
+	// ------------------
+
+	const uint32 NumIndices = MeshData.SourceMeshData ? MeshData.SourceMeshData->GetNumIndices() : MeshData.LODModel->IndexBuffer.GetNumIndices();
+	const int64 EmbreeUsage = MeshRepresentation::MemoryEstimateForEmbreeScene(NumIndices);
+
+	// ------------------
+
+	const int64 Total = MeshUtilTotal + EmbreeUsage;
+	return Total;
+}
+
 static void BuildSignedDistanceField(
 	FString MeshName,
 	const FEmbreeScene& EmbreeScene,
diff --git a/Engine/Source/Developer/MeshUtilities/Private/MeshRepresentationCommon.cpp b/Engine/Source/Developer/MeshUtilities/Private/MeshRepresentationCommon.cpp
--- a/Engine/Source/Developer/MeshUtilities/Private/MeshRepresentationCommon.cpp
+++ b/Engine/Source/Developer/MeshUtilities/Private/MeshRepresentationCommon.cpp
@@ -8,6 +8,17 @@
 #include "MeshUtilitiesPrivate.h"
 #include "DerivedMeshDataTaskUtils.h"
 
+#if USE_EMBREE
+
+static TAutoConsoleVariable<int32> CVarMemoryEstimateFactor(
+	TEXT("r.Embree.MemoryEstimateFactor"), 60,
+	TEXT("Configurable ratio for the memory used by embree"),
+	ECVF_ReadOnly
+); 
+
+#endif
+
+
 static FVector3f UniformSampleHemisphere(FVector2D Uniforms)
 {
 	Uniforms = Uniforms * 2.0f - 1.0f;
@@ -139,6 +150,16 @@ static bool EmbreeMemoryMonitor(void* userPtr, ssize_t bytes, bool post)
 }
 #endif
 
+int64_t MeshRepresentation::MemoryEstimateForEmbreeScene(uint64_t IndexCount)
+{
+	// This value was observed by breakpointing VmAlloc and viewing the allocations made
+	const int64_t EmbreeDeviceUsage = 1024 * 1024 * 16; 
+
+	// Estimate based of measuring and correlating actual memory usage with various mesh properties. 
+	// Strong correlation between number of indices and memory usage. 1mb bias added to catch edge cases with very low index counts.
+	return (CVarMemoryEstimateFactor.GetValueOnAnyThread() * IndexCount) + (1024 * 1024) + EmbreeDeviceUsage;
+}
+
 void MeshRepresentation::SetupEmbreeScene(FString MeshName, bool bGenerateAsIfTwoSided, FEmbreeScene& OutEmbreeScene)
 {
 	OutEmbreeScene.MeshName = MeshName;
diff --git a/Engine/Source/Developer/MeshUtilities/Private/MeshRepresentationCommon.h b/Engine/Source/Developer/MeshUtilities/Private/MeshRepresentationCommon.h
--- a/Engine/Source/Developer/MeshUtilities/Private/MeshRepresentationCommon.h
+++ b/Engine/Source/Developer/MeshUtilities/Private/MeshRepresentationCommon.h
@@ -156,6 +156,8 @@ namespace MeshRepresentation
 
 	void DeleteEmbreeScene(FEmbreeScene& EmbreeScene);
 
+	int64_t MemoryEstimateForEmbreeScene(uint64_t IndexCount);
+
 	bool AddMeshDataToEmbreeScene(FEmbreeScene& EmbreeScene, const FMeshDataForDerivedDataTask& MeshData, bool bIncludeTranslucentTriangles);
 
 };
diff --git a/Engine/Source/Developer/MeshUtilities/Private/MeshUtilitiesPrivate.h b/Engine/Source/Developer/MeshUtilities/Private/MeshUtilitiesPrivate.h
--- a/Engine/Source/Developer/MeshUtilities/Private/MeshUtilitiesPrivate.h
+++ b/Engine/Source/Developer/MeshUtilities/Private/MeshUtilitiesPrivate.h
@@ -55,6 +55,12 @@ private:
 
 	virtual bool GenerateStaticMeshLODs(UStaticMesh* StaticMesh, const FStaticMeshLODGroup& LODGroup) override;
 
+	
+	virtual int64_t FMeshUtilities::MemoryForGenerateSignedDistanceFieldVolumeData(
+		const FMeshDataForDerivedDataTask& MeshData,
+		float DistanceFieldResolutionScale
+	);
+
 	/** Builds a signed distance field volume for the given mesh. */
 	virtual void GenerateSignedDistanceFieldVolumeData(
 		FString MeshName,
diff --git a/Engine/Source/Developer/MeshUtilities/Public/MeshUtilities.h b/Engine/Source/Developer/MeshUtilities/Public/MeshUtilities.h
--- a/Engine/Source/Developer/MeshUtilities/Public/MeshUtilities.h
+++ b/Engine/Source/Developer/MeshUtilities/Public/MeshUtilities.h
@@ -138,6 +138,11 @@ public:
 		bool bGenerateAsIfTwoSided,
 		class FDistanceFieldVolumeData& OutData) = 0;
 
+	virtual int64_t MemoryForGenerateSignedDistanceFieldVolumeData(
+		const FMeshDataForDerivedDataTask& MeshData,
+		float DistanceFieldResolutionScale
+	) = 0;
+
 	UE_DEPRECATED(5.6, "Provide FMeshDataForDerivedDataTask instead.")
 	virtual void GenerateSignedDistanceFieldVolumeData(
 		FString MeshName,
diff --git a/Engine/Source/Runtime/Engine/Private/DistanceFieldAtlas.cpp b/Engine/Source/Runtime/Engine/Private/DistanceFieldAtlas.cpp
--- a/Engine/Source/Runtime/Engine/Private/DistanceFieldAtlas.cpp
+++ b/Engine/Source/Runtime/Engine/Private/DistanceFieldAtlas.cpp
@@ -120,6 +120,12 @@ static FAutoConsoleVariableRef CVarHFVisibilityAtlasDownSampleLevel(
 	TEXT("Max number of times a suballocation can be down-sampled"),
 	ECVF_RenderThreadSafe);
 
+static TAutoConsoleVariable<bool> CVarDistanceFieldsEstimate(
+	TEXT("r.CookMemory.DistanceFieldsEstimate"), true,
+	TEXT("Enable memory estimates for distance field asset cooks"),
+	ECVF_ReadOnly
+); 
+
 TGlobalResource<FLandscapeTextureAtlas> GHeightFieldTextureAtlas(FLandscapeTextureAtlas::SAT_Height);
 TGlobalResource<FLandscapeTextureAtlas> GHFVisibilityTextureAtlas(FLandscapeTextureAtlas::SAT_Visibility);
 
@@ -515,12 +521,32 @@ void FDistanceFieldAsyncQueue::CancelAndDeleteTask(const TSet<FAsyncDistanceFiel
 		delete Task;
 	}
 }
-
+ 
 void FDistanceFieldAsyncQueue::StartBackgroundTask(FAsyncDistanceFieldTask* Task)
 {
 	check(Task->AsyncTask == nullptr);
 	Task->AsyncTask = MakeUnique<FAsyncTask<FAsyncDistanceFieldTaskWorker>>(*Task);
-	int64 RequiredMemory = -1; // @todo RequiredMemory
+
+	int64 RequiredMemory = -1;
+
+#if WITH_EDITOR
+	if (CVarDistanceFieldsEstimate.GetValueOnAnyThread())
+	{
+		const FStaticMeshLODResources& LODModel = Task->GenerateSource->GetRenderData()->LODResources[0];
+
+		FMeshDataForDerivedDataTask MeshData;
+		MeshData.SourceMeshData = &Task->SourceMeshData;
+		MeshData.LODModel = &LODModel;
+		MeshData.SectionData = Task->SectionData;
+		MeshData.Bounds = (FBoxSphereBounds3f)Task->GenerateSource->GetRenderData()->Bounds;
+
+
+		RequiredMemory = MeshUtilities.load(std::memory_order_relaxed)->MemoryForGenerateSignedDistanceFieldVolumeData(
+			MeshData,
+			Task->DistanceFieldResolutionScale);
+	}
+#endif
+
 	Task->AsyncTask->StartBackgroundTask(ThreadPool.Get(), EQueuedWorkPriority::Lowest, EQueuedWorkFlags::DoNotRunInsideBusyWait, RequiredMemory, TEXT("DistanceField") );
 }
 
