Fix for compiler error, "extra qualification on member 'MemoryForGenerateSignedDistanceFieldVolumeData'"
#presubmit skip

diff --git a/Engine/Source/Developer/MeshUtilities/Private/MeshUtilitiesPrivate.h b/Engine/Source/Developer/MeshUtilities/Private/MeshUtilitiesPrivate.h
--- a/Engine/Source/Developer/MeshUtilities/Private/MeshUtilitiesPrivate.h
+++ b/Engine/Source/Developer/MeshUtilities/Private/MeshUtilitiesPrivate.h
@@ -56,7 +56,7 @@ private:
 	virtual bool GenerateStaticMeshLODs(UStaticMesh* StaticMesh, const FStaticMeshLODGroup& LODGroup) override;
 
 	
-	virtual int64_t FMeshUtilities::MemoryForGenerateSignedDistanceFieldVolumeData(
+	virtual int64_t MemoryForGenerateSignedDistanceFieldVolumeData(
 		const FMeshDataForDerivedDataTask& MeshData,
 		float DistanceFieldResolutionScale
 	);
