Skeletal Mesh: Make clearing ASW more targeted when rebuilding the render mesh.
- Don't force-clear vertex attribtues, since that happens in FLODUtilities::UpdateLODInfoVertexAttributes.

#rnx

diff --git a/Engine/Source/Developer/MeshBuilder/Private/SkeletalMeshBuilder.cpp b/Engine/Source/Developer/MeshBuilder/Private/SkeletalMeshBuilder.cpp
--- a/Engine/Source/Developer/MeshBuilder/Private/SkeletalMeshBuilder.cpp
+++ b/Engine/Source/Developer/MeshBuilder/Private/SkeletalMeshBuilder.cpp
@@ -338,6 +338,9 @@ public:
 
 	void UpdateAlternateSkinWeights()
 	{
+		// Clear out any existing alternate skin weights from the working LOD model. We will be fully rebuilding them below.
+		GetLODModel().SkinWeightProfiles.Reset();
+		
 		for (const FSkinWeightProfileInfo& ProfileInfo : SkeletalMesh->GetSkinWeightProfiles())
 		{			
 			FLODUtilities::UpdateAlternateSkinWeights(SkeletalMesh, ProfileInfo.Name, LODIndex, Options);
diff --git a/Engine/Source/Developer/MeshUtilities/Private/MeshUtilities.cpp b/Engine/Source/Developer/MeshUtilities/Private/MeshUtilities.cpp
--- a/Engine/Source/Developer/MeshUtilities/Private/MeshUtilities.cpp
+++ b/Engine/Source/Developer/MeshUtilities/Private/MeshUtilities.cpp
@@ -796,11 +796,6 @@ void FMeshUtilities::BuildSkeletalModelFromChunks(
 	LODModel.NumVertices = 0;
 	LODModel.IndexBuffer.Empty();
 
-	// Since we're building a new model, clear out skin weight profiles and vertex attributes. Those can be recreated by
-	// FLODUtilities::UpdateAlternateSkinWeights & FLODUtilities::UpdateLODInfoVertexAttributes as needed.
-	LODModel.SkinWeightProfiles.Empty();
-	LODModel.VertexAttributes.Empty();
-	
 	// Setup the section and chunk arrays on the model.
 	for (int32 ChunkIndex = 0; ChunkIndex < Chunks.Num(); ++ChunkIndex)
 	{
