Fix TArray reserve calculation in GenerateSignedDistanceFieldVolumeData

diff --git a/Engine/Source/Developer/MeshUtilities/Private/MeshDistanceFieldUtilities.cpp b/Engine/Source/Developer/MeshUtilities/Private/MeshDistanceFieldUtilities.cpp
--- a/Engine/Source/Developer/MeshUtilities/Private/MeshDistanceFieldUtilities.cpp
+++ b/Engine/Source/Developer/MeshUtilities/Private/MeshDistanceFieldUtilities.cpp
@@ -448,7 +448,7 @@ static void BuildSignedDistanceField(
 		const FVector2f DistanceFieldToVolumeScaleBias(2.0f * MaxDistanceForEncoding, -MaxDistanceForEncoding);
 
 		TArray<FSparseMeshDistanceFieldAsyncTask> AsyncTasks;
-		AsyncTasks.Reserve(IndirectionDimensions.X * IndirectionDimensions.Y * IndirectionDimensions.Z / 8);
+		AsyncTasks.Reserve(IndirectionDimensions.X * IndirectionDimensions.Y * IndirectionDimensions.Z);
 
 		for (int32 ZIndex = 0; ZIndex < IndirectionDimensions.Z; ZIndex++)
 		{
