compile fixes when USE_EMBREE is 0

#rnx

diff --git a/Engine/Source/Developer/MeshUtilities/Private/MeshRepresentationCommon.cpp b/Engine/Source/Developer/MeshUtilities/Private/MeshRepresentationCommon.cpp
--- a/Engine/Source/Developer/MeshUtilities/Private/MeshRepresentationCommon.cpp
+++ b/Engine/Source/Developer/MeshUtilities/Private/MeshRepresentationCommon.cpp
@@ -152,12 +152,16 @@ static bool EmbreeMemoryMonitor(void* userPtr, ssize_t bytes, bool post)
 
 int64_t MeshRepresentation::MemoryEstimateForEmbreeScene(uint64_t IndexCount)
 {
+#if USE_EMBREE
 	// This value was observed by breakpointing VmAlloc and viewing the allocations made
 	const int64_t EmbreeDeviceUsage = 1024 * 1024 * 16; 
 
 	// Estimate based of measuring and correlating actual memory usage with various mesh properties. 
 	// Strong correlation between number of indices and memory usage. 1mb bias added to catch edge cases with very low index counts.
 	return (CVarMemoryEstimateFactor.GetValueOnAnyThread() * IndexCount) + (1024 * 1024) + EmbreeDeviceUsage;
+#else
+	return 0;
+#endif
 }
 
 void MeshRepresentation::SetupEmbreeScene(FString MeshName, bool bGenerateAsIfTwoSided, FEmbreeScene& OutEmbreeScene)
diff --git a/Engine/Source/Developer/MeshUtilities/Private/MeshRepresentationCommon.h b/Engine/Source/Developer/MeshUtilities/Private/MeshRepresentationCommon.h
--- a/Engine/Source/Developer/MeshUtilities/Private/MeshRepresentationCommon.h
+++ b/Engine/Source/Developer/MeshUtilities/Private/MeshRepresentationCommon.h
@@ -13,6 +13,7 @@
 	#include <embree3/rtcore_ray.h>
 #endif
 #else
+	#define RTC_INVALID_GEOMETRY_ID ((unsigned int)-1)
 	typedef void* RTCDevice;
 	typedef void* RTCScene;
 	typedef void* RTCGeometry;
