[SkeletalMesh] Make sure skin weights are sorted by both index and weight when compared

Previously, verts that have the same weights in base vs in a alt profile but have different ordering of the weights are considered to have different weights and thus were incorrectly inserted into the alt profile causing visual artifacts when applying a 2-layer skin weight profile stack

diff --git a/Engine/Source/Developer/MeshUtilities/Private/MeshUtilities.cpp b/Engine/Source/Developer/MeshUtilities/Private/MeshUtilities.cpp
--- a/Engine/Source/Developer/MeshUtilities/Private/MeshUtilities.cpp
+++ b/Engine/Source/Developer/MeshUtilities/Private/MeshUtilities.cpp
@@ -5371,13 +5371,48 @@ void FMeshUtilities::GenerateRuntimeSkinWeightData(
 			const FSoftSkinVertex& TargetVertex = TargetVertices[VertexIndex];
 
 			bool bIsDifferent = false;
-			for (int32 InfluenceIndex = 0; InfluenceIndex < NumInfluences; ++InfluenceIndex)
+			// Weights are generally sorted from high to low, but they are not guaranteed to be sorted by bone index.
+			// Since we want to know if two list of weights are the same, we have to sort them again by both index and weight
 			{
-				if (SourceSkinWeight.InfluenceBones[InfluenceIndex] != TargetVertex.InfluenceBones[InfluenceIndex]
-					|| SourceSkinWeight.InfluenceWeights[InfluenceIndex] != TargetVertex.InfluenceWeights[InfluenceIndex])
+				struct FInfluence
 				{
-					bIsDifferent = true;
-					break;
+					uint16 BoneIndex = 0;
+					uint16 Weight = 0;
+
+					bool operator < (const FInfluence& Other) const
+					{
+						if (Weight == Other.Weight)
+						{
+							return BoneIndex < Other.BoneIndex;
+						}
+						
+						return Weight > Other.Weight;
+					}
+
+					bool operator == (const FInfluence& Other) const = default;
+				};
+
+
+				TArray<FInfluence, TInlineAllocator<MAX_TOTAL_INFLUENCES>> SourceWeightArray;
+				TArray<FInfluence, TInlineAllocator<MAX_TOTAL_INFLUENCES>> TargetWeightArray;
+			
+				for (int32 InfluenceIndex = 0; InfluenceIndex < NumInfluences; ++InfluenceIndex)
+				{
+					SourceWeightArray.Add({SourceSkinWeight.InfluenceBones[InfluenceIndex], SourceSkinWeight.InfluenceWeights[InfluenceIndex]});
+					TargetWeightArray.Add({TargetVertex.InfluenceBones[InfluenceIndex], TargetVertex.InfluenceWeights[InfluenceIndex]});
+				}
+
+				Algo::Sort(SourceWeightArray);
+				Algo::Sort(TargetWeightArray);
+			
+			
+				for (int32 InfluenceIndex = 0; InfluenceIndex < NumInfluences; ++InfluenceIndex)
+				{
+					if (SourceWeightArray[InfluenceIndex] != TargetWeightArray[InfluenceIndex])
+					{
+						bIsDifferent = true;
+						break;
+					}
 				}
 			}
 
