[Interchange] Fixed slow import of horizontal skeletal mesh with lots of vertices

- copied sorting from Static Mesh's - where there's still a plane(i.e. linear combination of xyz) but not trivial

#Synced-CL 50902575

diff --git a/Engine/Source/Developer/MeshUtilities/Private/MeshUtilities.cpp b/Engine/Source/Developer/MeshUtilities/Private/MeshUtilities.cpp
--- a/Engine/Source/Developer/MeshUtilities/Private/MeshUtilities.cpp
+++ b/Engine/Source/Developer/MeshUtilities/Private/MeshUtilities.cpp
@@ -4193,7 +4193,8 @@ public:
 				// Add an efficient way to find dupes of this vertex later for fast combining of vertices
 				FSkeletalMeshVertIndexAndZ IAndZ;
 				IAndZ.Index = RawIndex;
-				IAndZ.Z = Vertex.Position.Z;
+				const FVector3f& V = Vertex.Position;
+				IAndZ.Z = 0.30f * V.X + 0.33f * V.Y + 0.37f * V.Z;
 
 				VertIndexAndZ.Add(IAndZ);
 			}
